Niagara rebuild, first iteration; sim framework works, simple actor with a single simuation;

diff --git a/Engine/Source/Editor/GraphEditor/Private/NiagaraPins/SGraphPinVector4.cpp b/Engine/Source/Editor/GraphEditor/Private/NiagaraPins/SGraphPinVector4.cpp
new file mode 100644
--- /dev/null
+++ b/Engine/Source/Editor/GraphEditor/Private/NiagaraPins/SGraphPinVector4.cpp
@@ -0,0 +1,292 @@
+// Copyright 1998-2014 Epic Games, Inc. All Rights Reserved.
+
+
+#include "GraphEditorCommon.h"
+#include "SGraphPinVector4.h"
+
+
+#define LOCTEXT_NAMESPACE "VectorTextBox"
+
+
+//Class implementation to create 3 editable text boxes to represent vector/rotator graph pin
+class SVector4TextBox : public SCompoundWidget
+{
+public:
+
+	SLATE_BEGIN_ARGS(SVector4TextBox) {}
+	SLATE_ATTRIBUTE(FString, VisibleText_0)
+		SLATE_ATTRIBUTE(FString, VisibleText_1)
+		SLATE_ATTRIBUTE(FString, VisibleText_2)
+		SLATE_ATTRIBUTE(FString, VisibleText_3)
+		SLATE_EVENT(FOnFloatValueCommitted, OnFloatCommitted_Box_0)
+		SLATE_EVENT(FOnFloatValueCommitted, OnFloatCommitted_Box_1)
+		SLATE_EVENT(FOnFloatValueCommitted, OnFloatCommitted_Box_2)
+		SLATE_EVENT(FOnFloatValueCommitted, OnFloatCommitted_Box_3)
+		SLATE_END_ARGS()
+
+		//Construct editable text boxes with the appropriate getter & setter functions along with tool tip text
+		void Construct(const FArguments& InArgs)
+	{
+			VisibleText_0 = InArgs._VisibleText_0;
+			VisibleText_1 = InArgs._VisibleText_1;
+			VisibleText_2 = InArgs._VisibleText_2;
+			VisibleText_3 = InArgs._VisibleText_3;
+			const FLinearColor LabelClr = FLinearColor(1.f, 1.f, 1.f, 0.4f);
+
+			this->ChildSlot
+				[
+					SNew(SVerticalBox)
+					+ SVerticalBox::Slot()
+					.AutoHeight()
+					.Padding(0)
+					[
+						SNew(SHorizontalBox)
+						+ SHorizontalBox::Slot()
+						.AutoWidth().Padding(2).HAlign(HAlign_Fill)
+						[
+							//Create Text box 0 
+							SNew(SNumericEntryBox<float>)
+							.LabelVAlign(VAlign_Center)
+							.Label()
+							[
+								SNew(STextBlock)
+								.Font(FEditorStyle::GetFontStyle("Graph.VectorEditableTextBox"))
+								.Text(LOCTEXT("VectorNodeXAxisValueLabel", "X"))
+								.ColorAndOpacity(LabelClr)
+							]
+							.Value(this, &SVector4TextBox::GetTypeInValue_0)
+								.OnValueCommitted(InArgs._OnFloatCommitted_Box_0)
+								.Font(FEditorStyle::GetFontStyle("Graph.VectorEditableTextBox"))
+								.UndeterminedString(LOCTEXT("MultipleValues", "Multiple Values"))
+								.ToolTipText(LOCTEXT("VectorNodeXAxisValueLabel_ToolTip", "X value"))
+								.EditableTextBoxStyle(&FEditorStyle::GetWidgetStyle<FEditableTextBoxStyle>("Graph.VectorEditableTextBox"))
+								.BorderForegroundColor(FLinearColor::White)
+								.BorderBackgroundColor(FLinearColor::White)
+						]
+						+ SHorizontalBox::Slot()
+							.AutoWidth().Padding(2).HAlign(HAlign_Fill)
+							[
+								//Create Text box 1
+								SNew(SNumericEntryBox<float>)
+								.LabelVAlign(VAlign_Center)
+								.Label()
+								[
+									SNew(STextBlock)
+									.Font(FEditorStyle::GetFontStyle("Graph.VectorEditableTextBox"))
+									.Text(LOCTEXT("VectorNodeYAxisValueLabel", "Y"))
+									.ColorAndOpacity(LabelClr)
+								]
+								.Value(this, &SVector4TextBox::GetTypeInValue_1)
+									.OnValueCommitted(InArgs._OnFloatCommitted_Box_1)
+									.Font(FEditorStyle::GetFontStyle("Graph.VectorEditableTextBox"))
+									.UndeterminedString(LOCTEXT("MultipleValues", "Multiple Values"))
+									.ToolTipText(LOCTEXT("VectorNodeYAxisValueLabel_ToolTip", "Y value"))
+									.EditableTextBoxStyle(&FEditorStyle::GetWidgetStyle<FEditableTextBoxStyle>("Graph.VectorEditableTextBox"))
+									.BorderForegroundColor(FLinearColor::White)
+									.BorderBackgroundColor(FLinearColor::White)
+							]
+						+ SHorizontalBox::Slot()
+							.AutoWidth().Padding(2).HAlign(HAlign_Fill)
+							[
+								//Create Text box 2
+								SNew(SNumericEntryBox<float>)
+								.LabelVAlign(VAlign_Center)
+								.Label()
+								[
+									SNew(STextBlock)
+									.Font(FEditorStyle::GetFontStyle("Graph.VectorEditableTextBox"))
+									.Text(LOCTEXT("VectorNodeZAxisValueLabel", "Z"))
+									.ColorAndOpacity(LabelClr)
+								]
+								.Value(this, &SVector4TextBox::GetTypeInValue_2)
+									.OnValueCommitted(InArgs._OnFloatCommitted_Box_2)
+									.Font(FEditorStyle::GetFontStyle("Graph.VectorEditableTextBox"))
+									.UndeterminedString(LOCTEXT("MultipleValues", "Multiple Values"))
+									.ToolTipText(LOCTEXT("VectorNodeZAxisValueLabel_ToolTip", "Z value"))
+									.EditableTextBoxStyle(&FEditorStyle::GetWidgetStyle<FEditableTextBoxStyle>("Graph.VectorEditableTextBox"))
+									.BorderForegroundColor(FLinearColor::White)
+									.BorderBackgroundColor(FLinearColor::White)
+							]
+						+ SHorizontalBox::Slot()
+							.AutoWidth().Padding(2).HAlign(HAlign_Fill)
+							[
+								//Create Text box 3
+								SNew(SNumericEntryBox<float>)
+								.LabelVAlign(VAlign_Center)
+								.Label()
+								[
+									SNew(STextBlock)
+									.Font(FEditorStyle::GetFontStyle("Graph.VectorEditableTextBox"))
+									.Text(LOCTEXT("VectorNodeZAxisValueLabel", "W"))
+									.ColorAndOpacity(LabelClr)
+								]
+								.Value(this, &SVector4TextBox::GetTypeInValue_3)
+									.OnValueCommitted(InArgs._OnFloatCommitted_Box_3)
+									.Font(FEditorStyle::GetFontStyle("Graph.VectorEditableTextBox"))
+									.UndeterminedString(LOCTEXT("MultipleValues", "Multiple Values"))
+									.ToolTipText(LOCTEXT("VectorNodeZAxisValueLabel_ToolTip", "W value"))
+									.EditableTextBoxStyle(&FEditorStyle::GetWidgetStyle<FEditableTextBoxStyle>("Graph.VectorEditableTextBox"))
+									.BorderForegroundColor(FLinearColor::White)
+									.BorderBackgroundColor(FLinearColor::White)
+							]
+
+					]
+				];
+		}
+
+private:
+
+	//Get value for text box 0
+	TOptional<float> GetTypeInValue_0() const
+	{
+		return FCString::Atof(*(VisibleText_0.Get()));
+	}
+
+	//Get value for text box 1
+	TOptional<float> GetTypeInValue_1() const
+	{
+		return FCString::Atof(*(VisibleText_1.Get()));
+	}
+
+	//Get value for text box 2
+	TOptional<float> GetTypeInValue_2() const
+	{
+		return FCString::Atof(*(VisibleText_2.Get()));
+	}
+
+	//Get value for text box 3
+	TOptional<float> GetTypeInValue_3() const
+	{
+		return FCString::Atof(*(VisibleText_3.Get()));
+	}
+
+	TAttribute< FString >	VisibleText_0;
+	TAttribute< FString >	VisibleText_1;
+	TAttribute< FString >	VisibleText_2;
+	TAttribute< FString >	VisibleText_3;
+};
+
+
+
+
+/************************************************************************/
+/*                    SGraphPinVector4 class implementation              */
+/************************************************************************/
+
+void SGraphPinVector4::Construct(const FArguments& InArgs, UEdGraphPin* InGraphPinObj)
+{
+	SGraphPin::Construct(SGraphPin::FArguments(), InGraphPinObj);
+}
+
+TSharedRef<SWidget>	SGraphPinVector4::GetDefaultValueWidget()
+{
+	UScriptStruct* RotatorStruct = FindObjectChecked<UScriptStruct>(UObject::StaticClass(), TEXT("Rotator"));
+
+	//Create widget
+	return	SNew(SVector4TextBox)
+		.VisibleText_0(this, &SGraphPinVector4::GetCurrentValue_0)
+		.VisibleText_1(this, &SGraphPinVector4::GetCurrentValue_1)
+		.VisibleText_2(this, &SGraphPinVector4::GetCurrentValue_2)
+		.VisibleText_3(this, &SGraphPinVector4::GetCurrentValue_3)
+		.Visibility(this, &SGraphPin::GetDefaultValueVisibility)
+		.OnFloatCommitted_Box_0(this, &SGraphPinVector4::OnChangedValueTextBox_0)
+		.OnFloatCommitted_Box_1(this, &SGraphPinVector4::OnChangedValueTextBox_1)
+		.OnFloatCommitted_Box_2(this, &SGraphPinVector4::OnChangedValueTextBox_2)
+		.OnFloatCommitted_Box_3(this, &SGraphPinVector4::OnChangedValueTextBox_3);
+}
+
+//Rotator is represented as X->Roll, Y->Pitch, Z->Yaw
+
+FString SGraphPinVector4::GetCurrentValue_0() const
+{
+	//Text box 0: Rotator->Roll, Vector->X
+	return GetValue(TextBox_0);
+}
+
+FString SGraphPinVector4::GetCurrentValue_1() const
+{
+	//Text box 1: Rotator->Pitch, Vector->Y
+	return GetValue(TextBox_1);
+}
+
+FString SGraphPinVector4::GetCurrentValue_2() const
+{
+	//Text box 2: Rotator->Yaw, Vector->Z
+	return GetValue(TextBox_2);
+}
+
+FString SGraphPinVector4::GetCurrentValue_3() const
+{
+	//Text box 3: Vector->W
+	return GetValue(TextBox_3);
+}
+
+FString SGraphPinVector4::GetValue(ETextBoxIndex Index) const
+{
+	FString DefaultString = GraphPinObj->GetDefaultAsString();
+	TArray<FString> ResultString;
+
+	//Parse string to split its contents separated by ','
+	DefaultString.Trim();
+	DefaultString.TrimTrailing();
+	DefaultString.ParseIntoArray(&ResultString, TEXT(","), true);
+
+	if (Index < ResultString.Num())
+	{
+		return ResultString[Index];
+	}
+	else
+	{
+		return FString(TEXT("0"));
+	}
+}
+
+
+void SGraphPinVector4::OnChangedValueTextBox_0(float NewValue, ETextCommit::Type CommitInfo)
+{
+	const FString ValueStr = FString::Printf(TEXT("%f"), NewValue);
+
+	FString DefaultValue;
+	//Update X value
+	DefaultValue = ValueStr + FString(TEXT(",")) + GetValue(TextBox_1) + FString(TEXT(",")) + GetValue(TextBox_2) + GetValue(TextBox_3);
+	//Set new default value
+	GraphPinObj->GetSchema()->TrySetDefaultValue(*GraphPinObj, DefaultValue);
+}
+
+void SGraphPinVector4::OnChangedValueTextBox_1(float NewValue, ETextCommit::Type CommitInfo)
+{
+	const FString ValueStr = FString::Printf(TEXT("%f"), NewValue);
+
+	FString DefaultValue;
+	//Update Y value
+	DefaultValue = GetValue(TextBox_0) + FString(TEXT(",")) + ValueStr + FString(TEXT(",")) + GetValue(TextBox_2) + FString(TEXT(",")) + GetValue(TextBox_3);
+	//Set new default value
+	GraphPinObj->GetSchema()->TrySetDefaultValue(*GraphPinObj, DefaultValue);
+}
+
+void SGraphPinVector4::OnChangedValueTextBox_2(float NewValue, ETextCommit::Type CommitInfo)
+{
+	const FString ValueStr = FString::Printf(TEXT("%f"), NewValue);
+
+	FString DefaultValue;
+	//Update Z value
+	DefaultValue = GetValue(TextBox_0) + FString(TEXT(",")) + GetValue(TextBox_1) + FString(TEXT(",")) + ValueStr + FString(TEXT(",")) + GetValue(TextBox_2);
+	//Set new default value
+	GraphPinObj->GetSchema()->TrySetDefaultValue(*GraphPinObj, DefaultValue);
+}
+
+void SGraphPinVector4::OnChangedValueTextBox_3(float NewValue, ETextCommit::Type CommitInfo)
+{
+	const FString ValueStr = FString::Printf(TEXT("%f"), NewValue);
+
+	FString DefaultValue;
+	//Update W value
+	DefaultValue = GetValue(TextBox_0) + FString(TEXT(",")) + GetValue(TextBox_1) + FString(TEXT(",")) + GetValue(TextBox_2) + FString(TEXT(",")) + ValueStr;
+
+	//Set new default value
+	GraphPinObj->GetSchema()->TrySetDefaultValue(*GraphPinObj, DefaultValue);
+}
+
+
+
+#undef LOCTEXT_NAMESPACE
diff --git a/Engine/Source/Editor/GraphEditor/Private/NiagaraPins/SGraphPinVector4.h b/Engine/Source/Editor/GraphEditor/Private/NiagaraPins/SGraphPinVector4.h
new file mode 100644
--- /dev/null
+++ b/Engine/Source/Editor/GraphEditor/Private/NiagaraPins/SGraphPinVector4.h
@@ -0,0 +1,60 @@
+// Copyright 1998-2014 Epic Games, Inc. All Rights Reserved.
+
+#pragma once
+
+class SGraphPinVector4 : public SGraphPin
+{
+public:
+	SLATE_BEGIN_ARGS(SGraphPinVector4) {}
+	SLATE_END_ARGS()
+
+		void Construct(const FArguments& InArgs, UEdGraphPin* InGraphPinObj);
+
+protected:
+
+	/**
+	*	Function to create class specific widget.
+	*
+	*	@return Reference to the newly created widget object
+	*/
+	virtual TSharedRef<SWidget>	GetDefaultValueWidget() override;
+
+private:
+
+	// Enum values to identify text boxes.
+	enum ETextBoxIndex
+	{
+		TextBox_0,
+		TextBox_1,
+		TextBox_2,
+		TextBox_3,
+	};
+
+	FString GetCurrentValue_0() const;
+	FString GetCurrentValue_1() const;
+	FString GetCurrentValue_2() const;
+	FString GetCurrentValue_3() const;
+
+	/*
+	*	Function to getch current value based on text box index value
+	*
+	*	@param: Text box index
+	*
+	*	@return current string value
+	*/
+	FString GetValue(ETextBoxIndex Index) const;
+
+	/*
+	*	Function to store value when text box value in modified
+	*
+	*	@param 0: Updated Float Value
+	*/
+	void OnChangedValueTextBox_0(float NewValue, ETextCommit::Type CommitInfo);
+	void OnChangedValueTextBox_1(float NewValue, ETextCommit::Type CommitInfo);
+	void OnChangedValueTextBox_2(float NewValue, ETextCommit::Type CommitInfo);
+	void OnChangedValueTextBox_3(float NewValue, ETextCommit::Type CommitInfo);
+
+private:
+	/** Flag is true if the widget is used to represent a rotator; false otherwise */
+	bool bIsRotator;
+};
diff --git a/Engine/Source/Editor/GraphEditor/Private/NodeFactory.cpp b/Engine/Source/Editor/GraphEditor/Private/NodeFactory.cpp
--- a/Engine/Source/Editor/GraphEditor/Private/NodeFactory.cpp
+++ b/Engine/Source/Editor/GraphEditor/Private/NodeFactory.cpp
@@ -59,6 +59,7 @@
 #include "KismetPins/SGraphPinNameList.h"
 #include "KismetPins/SGraphPinVector.h"
 #include "KismetPins/SGraphPinVector2D.h"
+#include "NiagaraPins/SGraphPinVector4.h"
 #include "KismetPins/SGraphPinIndex.h"
 
 #include "SoundNodes/SGraphNodeSoundBase.h"
@@ -360,6 +361,16 @@ TSharedPtr<SGraphPin> FNodeFactory::CreatePinWidget(UEdGraphPin* InPin)
 		}
 	}
 
+	if (const UEdGraphSchema_Niagara* NSchema = Cast<const UEdGraphSchema_Niagara>(InPin->GetSchema()))
+	{
+		UScriptStruct* VectorStruct = FindObjectChecked<UScriptStruct>(UObject::StaticClass(), TEXT("Vector"));
+		if (InPin->PinType.PinSubCategoryObject == VectorStruct)
+		{
+			return SNew(SGraphPinVector4, InPin);
+		}
+	}
+
+
 	// If we didn't pick a custom pin widget, use an uncustomized basic pin
 	return SNew(SGraphPin, InPin);
 }
diff --git a/Engine/Source/Editor/NiagaraEditor/Private/NiagaraCompiler.cpp b/Engine/Source/Editor/NiagaraEditor/Private/NiagaraCompiler.cpp
--- a/Engine/Source/Editor/NiagaraEditor/Private/NiagaraCompiler.cpp
+++ b/Engine/Source/Editor/NiagaraEditor/Private/NiagaraCompiler.cpp
@@ -2,6 +2,7 @@
 
 #include "NiagaraEditorPrivatePCH.h"
 #include "Engine/NiagaraScript.h"
+#include "Engine/NiagaraConstants.h"
 #include "Components/NiagaraComponent.h"
 #include "CompilerResultsLog.h"
 #include "EdGraphUtilities.h"
@@ -17,11 +18,14 @@ struct FNiagaraExpr
 	uint8 OpIndex;
 	/** Source operands. */
 	int32 Src[3];
+	/* Source operand type vector (bit 0-2, 1=constant)*/
+	uint8 SrcOperandTypeVector;
 
 	/** Default constructor. */
 	FNiagaraExpr()
 	{
 		OpIndex = 0;
+		SrcOperandTypeVector = 0;
 		for (int32 i = 0; i < 3; i++)
 		{
 			Src[i] = INDEX_NONE;
@@ -35,7 +39,7 @@ struct FNiagaraCompilerContext
 	/** Expression. */
 	TArray<FNiagaraExpr> Expressions;
 	/** Constant table. */
-	TArray<float> Constants;
+	TArray<FVector4> Constants;
 	/** Constant names. */
 	TArray<FName> ConstantNames;
 	/** Map from pin to expression to prevent executing subtrees multiple times. */
@@ -49,11 +53,15 @@ struct FNiagaraCompilerContext
 	explicit FNiagaraCompilerContext(FCompilerResultsLog& InLog)
 		: Log(InLog)
 	{
-		// Setup built-in constants.
 		ConstantNames.Add(TEXT("__zero__"));
-		Constants.Add(0.0f);
-		ConstantNames.Add(TEXT("DeltaTime"));
-		Constants.Add(0.0f);
+		Constants.Add(FVector4(0.0f, 0.0f, 0.0f, 0.0f));
+
+		// Setup built-in constants.
+		for (uint32 i = 0; i < NiagaraConstants::NumBuiltinConstants; i++)
+		{
+			ConstantNames.Add(NiagaraConstants::ConstantNames[i]);
+			Constants.Add(FVector4(0.0f, 0.0f, 0.0f, 0.0f));
+		}
 	}
 };
 
@@ -146,6 +154,8 @@ static int32 EvaluateGraph(FNiagaraCompilerContext& Context, UEdGraphPin* Pin)
 
 					FNiagaraExpr NewOp;
 					NewOp.OpIndex = OpNode->OpIndex;
+					NewOp.SrcOperandTypeVector = 0;
+
 					for (int32 SrcIndex = 0; SrcIndex < 3; ++SrcIndex)
 					{
 						if (OpInfo.SrcTypes[SrcIndex] != VectorVM::EOpSrc::Invalid)
@@ -160,6 +170,7 @@ static int32 EvaluateGraph(FNiagaraCompilerContext& Context, UEdGraphPin* Pin)
 						NewOp.OpIndex++;
 					}
 
+
 					if (OpInfo.IsCommutative() && !IsValidOp(Context, NewOp))
 					{
 						int32 Temp = NewOp.Src[0];
@@ -171,6 +182,8 @@ static int32 EvaluateGraph(FNiagaraCompilerContext& Context, UEdGraphPin* Pin)
 							NewOp.OpIndex++;
 						}
 					}
+					check(NewOp.OpIndex < VectorVM::EOp::NumOpcodes);
+
 
 					if (IsValidOp(Context, NewOp))
 					{
@@ -200,14 +213,14 @@ static int32 EvaluateGraph(FNiagaraCompilerContext& Context, UEdGraphPin* Pin)
 			int32 ConstIndex = Context.ConstantNames.Find(ConstName);
 			if (ConstIndex == INDEX_NONE)
 			{
-				float ConstValue = FCString::Atof(*ConstString);
-				if (FMath::Abs(ConstValue) < SMALL_NUMBER)
-				{
-					ConstIndex = 0;
-				}
-				else
+				TArray<FString> ResultString;
+				ConstString.Trim();
+				ConstString.TrimTrailing();
+				ConstString.ParseIntoArray(&ResultString, TEXT(","), true);
+				ConstIndex = 0;
+				if (ResultString.Num() == 4)
 				{
-					ConstIndex = Context.Constants.Add(ConstValue);
+					ConstIndex = Context.Constants.Add(FVector4(FCString::Atof(*ResultString[0]), FCString::Atof(*ResultString[1]), FCString::Atof(*ResultString[2]), FCString::Atof(*ResultString[3])));
 					check(Context.ConstantNames.Num() == ConstIndex);
 					Context.ConstantNames.Add(ConstName);
 				}
@@ -259,7 +272,7 @@ void FNiagaraEditorModule::CompileScript(UNiagaraScript* ScriptToCompile)
 
 	// Traverse the node graph for each output and generate expressions as we go.
 	FNiagaraCompilerContext Context(MessageLog);
-	Source->GetUpdateOutputs(Context.Attributes);
+	Source->GetParticleAttributes(Context.Attributes);
 	TArray<int32> OutputExpressions;
 	OutputExpressions.AddUninitialized(Context.Attributes.Num());
 	FMemory::Memset(OutputExpressions.GetData(), 0xff, Context.Attributes.Num() * sizeof(int32));
@@ -337,9 +350,16 @@ void FNiagaraEditorModule::CompileScript(UNiagaraScript* ScriptToCompile)
 			{
 				// VectorVM::EOpSrc::Const.
 				Expr.Src[j] = ConstIndexFromExpressionIndex(Expr.Src[j]);
+				Expr.SrcOperandTypeVector |= 1<<j;
 			}
 		}
 
+		// now that we've figured out register/constant assignments, we revert back to the base opcode; the kernel Exec function will handle passing 
+		// the data in correctly based on the SrcOperandTypeVector. It would be good to change the compiler so we don't need the permutation
+		// opcodes at all, i.e. figure out source operand types without needing them in the OpInfo.
+		Expr.OpIndex = OpInfo.BaseOpcode;
+
+
 		int32 AttrIndex = INDEX_NONE;
 		if (OutputExpressions.Find(i, AttrIndex))
 		{
@@ -376,8 +396,12 @@ void FNiagaraEditorModule::CompileScript(UNiagaraScript* ScriptToCompile)
 		int32 DestRegister = RegisterAssignments[i];
 		check(DestRegister < VectorVM::MaxRegisters);
 
+		check(Expr.OpIndex < VectorVM::EOp::NumOpcodes);
+		check(Expr.OpIndex == OpInfo.BaseOpcode);
+
 		Code.Add(Expr.OpIndex);
 		Code.Add(DestRegister);
+		Code.Add(Expr.SrcOperandTypeVector);
 		for (int32 j = 0; j < 3; ++j)
 		{
 			if (OpInfo.SrcTypes[j] == VectorVM::EOpSrc::Invalid)
@@ -395,6 +419,7 @@ void FNiagaraEditorModule::CompileScript(UNiagaraScript* ScriptToCompile)
 				{
 					Code.Add(VectorVM::EOp::addi);
 					Code.Add(VectorVM::FirstOutputRegister + j);
+					Code.Add(0x0);	// all inputs are registers
 					Code.Add(DestRegister);
 					Code.Add(0);
 				}
diff --git a/Engine/Source/Editor/UnrealEd/Classes/Niagara/EdGraphSchema_Niagara.h b/Engine/Source/Editor/UnrealEd/Classes/Niagara/EdGraphSchema_Niagara.h
--- a/Engine/Source/Editor/UnrealEd/Classes/Niagara/EdGraphSchema_Niagara.h
+++ b/Engine/Source/Editor/UnrealEd/Classes/Niagara/EdGraphSchema_Niagara.h
@@ -49,6 +49,8 @@ class UEdGraphSchema_Niagara : public UEdGraphSchema
 	// Allowable PinType.PinCategory values
 	UPROPERTY()
 	FString PC_Float;
+	UPROPERTY()
+	FString PC_Struct;
 
 	// Begin EdGraphSchema interface
 	virtual void GetGraphContextActions(FGraphContextMenuBuilder& ContextMenuBuilder) const override;
diff --git a/Engine/Source/Editor/UnrealEd/Classes/Niagara/NiagaraNodeOp.h b/Engine/Source/Editor/UnrealEd/Classes/Niagara/NiagaraNodeOp.h
--- a/Engine/Source/Editor/UnrealEd/Classes/Niagara/NiagaraNodeOp.h
+++ b/Engine/Source/Editor/UnrealEd/Classes/Niagara/NiagaraNodeOp.h
@@ -22,3 +22,57 @@ class UNiagaraNodeOp : public UNiagaraNode
 	UNREALED_API static const TCHAR* InPinNames[3];
 };
 
+
+
+#if 0
+
+class SGraphNodeNiagara : public SGraphNode
+{
+public:
+	SLATE_BEGIN_ARGS(SGraphNodeNiagara){}
+	SLATE_END_ARGS()
+
+		void Construct(const FArguments& InArgs, class UNiagaraNode* InNode);
+
+	// SGraphNode interface
+	virtual void CreatePinWidgets() override;
+	// End of SGraphNode interface
+
+	// SNodePanel::SNode interface
+	virtual void MoveTo(const FVector2D& NewPosition, FNodeSet& NodeFilter) override;
+	// End of SNodePanel::SNode interface
+
+	UNiagaraNode* GetNiagaraGraphNode() const { return NiagaraNode; }
+
+protected:
+	// SGraphNode interface
+	virtual void AddPin(const TSharedRef<SGraphPin>& PinToAdd) override;
+	virtual void CreateBelowPinControls(TSharedPtr<SVerticalBox> MainBox) override;
+	virtual void SetDefaultTitleAreaWidget(TSharedRef<SOverlay> DefaultTitleAreaWidget) override;
+	virtual TSharedRef<SWidget> CreateNodeContentArea() override;
+	// End of SGraphNode interface
+
+	/** Creates a preview viewport if necessary */
+	//TSharedRef<SWidget> CreatePreviewWidget();
+
+	/** Returns visibility of Expression Preview viewport */
+	//EVisibility ExpressionPreviewVisibility() const;
+
+	/** Show/hide Expression Preview */
+	//void OnExpressionPreviewChanged(const ESlateCheckBoxState::Type NewCheckedState);
+
+	/** hidden == unchecked, shown == checked */
+	//ESlateCheckBoxState::Type IsExpressionPreviewChecked() const;
+
+	/** Up when shown, down when hidden */
+	//const FSlateBrush* GetExpressionPreviewArrow() const;
+
+private:
+	/** Slate viewport for rendering preview via custom slate element */
+	//TSharedPtr<FPreviewViewport> PreviewViewport;
+
+	/** Cached material graph node pointer to avoid casting */
+	UNiagaraNode* NiagaraNode;
+};
+
+#endif
\ No newline at end of file
diff --git a/Engine/Source/Editor/UnrealEd/Classes/Niagara/NiagaraScriptSource.h b/Engine/Source/Editor/UnrealEd/Classes/Niagara/NiagaraScriptSource.h
--- a/Engine/Source/Editor/UnrealEd/Classes/Niagara/NiagaraScriptSource.h
+++ b/Engine/Source/Editor/UnrealEd/Classes/Niagara/NiagaraScriptSource.h
@@ -17,9 +17,6 @@ class UNiagaraScriptSource : public UNiagaraScriptSourceBase
 	// UObject interface.
 	virtual void PostLoad() override;
 
-	/** Get the set of outputs from the Update script */
-	UNREALED_API void GetUpdateOutputs(TArray<FName>& ScalarOutputs);
-
-	/** Get the set of input attributes that can be used in the Update script */
-	UNREALED_API void GetUpdateInputs(TArray<FName>& ScalarInputs);
+	UNREALED_API void GetParticleAttributes(TArray<FName>& VectorOutputs);
+	UNREALED_API void GetEmitterAttributes(TArray<FName>& ScalarInputs);
 };
diff --git a/Engine/Source/Editor/UnrealEd/Private/Factories/NiagaraScriptFactory.cpp b/Engine/Source/Editor/UnrealEd/Private/Factories/NiagaraScriptFactory.cpp
--- a/Engine/Source/Editor/UnrealEd/Private/Factories/NiagaraScriptFactory.cpp
+++ b/Engine/Source/Editor/UnrealEd/Private/Factories/NiagaraScriptFactory.cpp
@@ -49,9 +49,10 @@ UObject* UNiagaraScriptFactoryNew::FactoryCreateNew(UClass* Class,UObject* InPar
 			UEdGraphPin* DeltaTimePin = DeltaTimeNode->FindPinChecked(DeltaTimeNode->AttrName.ToString());
 
 			TArray<FName> OutputNames;
-			Source->GetUpdateOutputs(OutputNames);
+			Source->GetParticleAttributes(OutputNames);
 
 			// Simple Euler integration.
+			/*
 			for(int32 i=0; i<3; i++)
 			{
 				FGraphNodeCreator<UNiagaraNodeGetAttr> PosNodeCreator(*CreatedGraph);
@@ -118,6 +119,7 @@ UObject* UNiagaraScriptFactoryNew::FactoryCreateNew(UClass* Class,UObject* InPar
 
 				OutputNode->FindPinChecked(TEXT("RelativeTime"))->MakeLinkTo(MadNode->FindPinChecked(TEXT("Result")));
 			}
+			*/
 
 			// Set pointer in script to source
 			NewScript->Source = Source;
diff --git a/Engine/Source/Editor/UnrealEd/Private/Niagara/EdGraphSchema_Niagara.cpp b/Engine/Source/Editor/UnrealEd/Private/Niagara/EdGraphSchema_Niagara.cpp
--- a/Engine/Source/Editor/UnrealEd/Private/Niagara/EdGraphSchema_Niagara.cpp
+++ b/Engine/Source/Editor/UnrealEd/Private/Niagara/EdGraphSchema_Niagara.cpp
@@ -96,6 +96,7 @@ UEdGraphSchema_Niagara::UEdGraphSchema_Niagara(const class FPostConstructInitial
 {
 
 	PC_Float = TEXT("float");
+	PC_Struct = TEXT("struct");
 }
 
 TSharedPtr<FNiagaraSchemaAction_NewNode> AddNewNodeAction(FGraphContextMenuBuilder& ContextMenuBuilder, const FString& Category, const FText& MenuDesc, const FString& Tooltip)
@@ -111,7 +112,8 @@ void UEdGraphSchema_Niagara::GetGraphContextActions(FGraphContextMenuBuilder& Co
 	UNiagaraScriptSource* Source = NiagaraGraph->GetSource();
 
 	TArray<FName> InputAttributeNames;
-	Source->GetUpdateInputs(InputAttributeNames);
+	Source->GetParticleAttributes(InputAttributeNames);
+	Source->GetEmitterAttributes(InputAttributeNames);
 	for(int32 i=0; i<InputAttributeNames.Num(); i++)
 	{
 		const FName AttrName = InputAttributeNames[i];
@@ -194,6 +196,10 @@ FLinearColor UEdGraphSchema_Niagara::GetPinTypeColor(const FEdGraphPinType& PinT
 	{
 		return Settings->FloatPinTypeColor;
 	}
+	else if (TypeString == PC_Struct)
+	{
+		return Settings->StructPinTypeColor;
+	}
 
 	return Settings->DefaultPinTypeColor;
 }
diff --git a/Engine/Source/Editor/UnrealEd/Private/Niagara/NiagaraNodeGetAttr.cpp b/Engine/Source/Editor/UnrealEd/Private/Niagara/NiagaraNodeGetAttr.cpp
--- a/Engine/Source/Editor/UnrealEd/Private/Niagara/NiagaraNodeGetAttr.cpp
+++ b/Engine/Source/Editor/UnrealEd/Private/Niagara/NiagaraNodeGetAttr.cpp
@@ -11,7 +11,8 @@ void UNiagaraNodeGetAttr::AllocateDefaultPins()
 {
 	const UEdGraphSchema_Niagara* Schema = GetDefault<UEdGraphSchema_Niagara>();
 
-	CreatePin(EGPD_Output, Schema->PC_Float, TEXT(""), NULL, false, false, AttrName.ToString());
+	UScriptStruct* VectorStruct = FindObjectChecked<UScriptStruct>(UObject::StaticClass(), TEXT("Vector"));
+	CreatePin(EGPD_Output, Schema->PC_Struct, TEXT(""), VectorStruct, false, false, AttrName.ToString());
 }
 
 FText UNiagaraNodeGetAttr::GetNodeTitle(ENodeTitleType::Type TitleType) const
diff --git a/Engine/Source/Editor/UnrealEd/Private/Niagara/NiagaraNodeOp.cpp b/Engine/Source/Editor/UnrealEd/Private/Niagara/NiagaraNodeOp.cpp
--- a/Engine/Source/Editor/UnrealEd/Private/Niagara/NiagaraNodeOp.cpp
+++ b/Engine/Source/Editor/UnrealEd/Private/Niagara/NiagaraNodeOp.cpp
@@ -22,19 +22,23 @@ void UNiagaraNodeOp::AllocateDefaultPins()
 	const UEdGraphSchema_Niagara* Schema = GetDefault<UEdGraphSchema_Niagara>();
 	VectorVM::FVectorVMOpInfo const& Info = VectorVM::GetOpCodeInfo(OpIndex);
 
+	UScriptStruct* VectorStruct = FindObjectChecked<UScriptStruct>(UObject::StaticClass(), TEXT("Vector"));
 	for (int32 SrcIndex = 0; SrcIndex < 3; ++SrcIndex)
 	{
 		VectorVM::EOpSrc::Type PinType = Info.SrcTypes[SrcIndex];
 		if (PinType != VectorVM::EOpSrc::Invalid)
 		{
-			UEdGraphPin* Pin = CreatePin(EGPD_Input, Schema->PC_Float, TEXT(""), NULL, false, false, InPinNames[SrcIndex]);
+			UEdGraphPin* Pin = CreatePin(EGPD_Input, Schema->PC_Struct, TEXT(""), VectorStruct, false, false, InPinNames[SrcIndex]);
 			check(Pin);
 			Pin->bDefaultValueIsIgnored = false;
+			Pin->bDefaultValueIsReadOnly = false;
 			Pin->bNotConnectable = false;
-			Pin->DefaultValue = TEXT("0.0");
+			Pin->DefaultValue = TEXT("1.0,1.0,1.0,1.0");
+			Pin->AutogeneratedDefaultValue = TEXT("1.0,1.0,1.0,1.0");
 		}
 	}
-	CreatePin(EGPD_Output, Schema->PC_Float, TEXT(""), NULL, false, false, TEXT("Result"));
+
+	CreatePin(EGPD_Output, Schema->PC_Struct, TEXT(""), VectorStruct, false, false, TEXT("Result"));
 }
 
 FText UNiagaraNodeOp::GetNodeTitle(ENodeTitleType::Type TitleType) const
@@ -47,3 +51,5 @@ FLinearColor UNiagaraNodeOp::GetNodeTitleColor() const
 {
 	return GetDefault<UGraphEditorSettings>()->FunctionCallNodeTitleColor;
 }
+
+
diff --git a/Engine/Source/Editor/UnrealEd/Private/Niagara/NiagaraNodeOutputUpdate.cpp b/Engine/Source/Editor/UnrealEd/Private/Niagara/NiagaraNodeOutputUpdate.cpp
--- a/Engine/Source/Editor/UnrealEd/Private/Niagara/NiagaraNodeOutputUpdate.cpp
+++ b/Engine/Source/Editor/UnrealEd/Private/Niagara/NiagaraNodeOutputUpdate.cpp
@@ -17,10 +17,12 @@ void UNiagaraNodeOutputUpdate::AllocateDefaultPins()
 	UNiagaraScriptSource* Source = GetSource();
 
 	TArray<FName> OutputNames;
-	Source->GetUpdateOutputs(OutputNames);
-	for(int32 i=0; i<OutputNames.Num(); i++)
+	Source->GetParticleAttributes(OutputNames);
+
+	UScriptStruct* VectorStruct = FindObjectChecked<UScriptStruct>(UObject::StaticClass(), TEXT("Vector"));
+	for (int32 i = 0; i < OutputNames.Num(); i++)
 	{
-		UEdGraphPin* Pin = CreatePin(EGPD_Input, Schema->PC_Float, TEXT(""), NULL, false, false, OutputNames[i].ToString());
+		UEdGraphPin* Pin = CreatePin(EGPD_Input, Schema->PC_Struct, TEXT(""), VectorStruct, false, false, OutputNames[i].ToString());
 		Pin->bDefaultValueIsIgnored = true;
 	}
 }
diff --git a/Engine/Source/Editor/UnrealEd/Private/Niagara/NiagaraScriptSource.cpp b/Engine/Source/Editor/UnrealEd/Private/Niagara/NiagaraScriptSource.cpp
--- a/Engine/Source/Editor/UnrealEd/Private/Niagara/NiagaraScriptSource.cpp
+++ b/Engine/Source/Editor/UnrealEd/Private/Niagara/NiagaraScriptSource.cpp
@@ -2,6 +2,7 @@
 
 #include "UnrealEd.h"
 #include "Engine/NiagaraScript.h"
+#include "Engine/NiagaraConstants.h"
 #include "NiagaraEditorModule.h"
 
 //////////////////////////////////////////////////////////////////////////
@@ -40,30 +41,23 @@ void UNiagaraScriptSource::PostLoad()
 	}
 }
 
-void UNiagaraScriptSource::GetUpdateOutputs(TArray<FName>& ScalarOutputs)
-{
-	ScalarOutputs.Empty();
-
-	ScalarOutputs.Add(FName(TEXT("Position.X")));
-	ScalarOutputs.Add(FName(TEXT("Position.Y")));
-	ScalarOutputs.Add(FName(TEXT("Position.Z")));
-
-	ScalarOutputs.Add(FName(TEXT("Velocity.X")));
-	ScalarOutputs.Add(FName(TEXT("Velocity.Y")));
-	ScalarOutputs.Add(FName(TEXT("Velocity.Z")));
-
-	ScalarOutputs.Add(FName(TEXT("Color.R")));
-	ScalarOutputs.Add(FName(TEXT("Color.G")));
-	ScalarOutputs.Add(FName(TEXT("Color.B")));
-	ScalarOutputs.Add(FName(TEXT("Color.A")));
 
-	ScalarOutputs.Add(FName(TEXT("Rotation")));
+void UNiagaraScriptSource::GetParticleAttributes(TArray<FName>& VectorOutputs)
+{
+	VectorOutputs.Empty();
 
-	ScalarOutputs.Add(FName(TEXT("RelativeTime")));
+	VectorOutputs.Add(FName(TEXT("Particle Position")));
+	VectorOutputs.Add(FName(TEXT("Particle Velocity")));
+	VectorOutputs.Add(FName(TEXT("Particle Color")));
+	VectorOutputs.Add(FName(TEXT("Particle Rotation")));
+	VectorOutputs.Add(FName(TEXT("Particle Age")));
 }
 
-void UNiagaraScriptSource::GetUpdateInputs(TArray<FName>& ScalarInputs)
+void UNiagaraScriptSource::GetEmitterAttributes(TArray<FName>& VectorInputs)
 {
-	GetUpdateOutputs(ScalarInputs);
-	ScalarInputs.Add(FName(TEXT("DeltaTime")));
+	for (uint32 i=0; i < NiagaraConstants::NumBuiltinConstants; i++)
+	{
+		VectorInputs.Add(NiagaraConstants::ConstantNames[i]);
+	}
 }
+
diff --git a/Engine/Source/Runtime/Engine/Classes/Components/NiagaraComponent.h b/Engine/Source/Runtime/Engine/Classes/Components/NiagaraComponent.h
--- a/Engine/Source/Runtime/Engine/Classes/Components/NiagaraComponent.h
+++ b/Engine/Source/Runtime/Engine/Classes/Components/NiagaraComponent.h
@@ -15,6 +15,8 @@ class ENGINE_API UNiagaraComponent : public UPrimitiveComponent
 	/** Script to run for these particles */
 	UPROPERTY(EditAnywhere, Category=NiagaraComponent)
 	class UNiagaraScript* UpdateScript;
+	UPROPERTY(EditAnywhere, Category = NiagaraComponent)
+	class UNiagaraScript* SpawnScript;
 
 	/** Material with which to render particles. */
 	UPROPERTY(EditAnywhere, Category=Rendering)
@@ -24,6 +26,9 @@ class ENGINE_API UNiagaraComponent : public UPrimitiveComponent
 	UPROPERTY(EditAnywhere, Category=NiagaraComponent)
 	float SpawnRate;
 
+	/** age of the emitter in seconds */
+	float EmitterAge;
+
 	// Begin UActorComponent interface.
 	virtual void TickComponent(float DeltaTime, enum ELevelTick TickType, FActorComponentTickFunction *ThisTickFunction) override;
 protected:
diff --git a/Engine/Source/Runtime/Engine/Classes/Engine/NiagaraConstants.h b/Engine/Source/Runtime/Engine/Classes/Engine/NiagaraConstants.h
new file mode 100644
--- /dev/null
+++ b/Engine/Source/Runtime/Engine/Classes/Engine/NiagaraConstants.h
@@ -0,0 +1,26 @@
+
+/**  List of built-in constants
+*	The values of these enums correspond directly to indices for the compiler and interpreter
+*/
+namespace NiagaraConstants
+{
+	enum {
+		DeltaSeconds = 0,
+		EmitterPosition,
+		EmitterAge,
+		ActorXAxis,
+		ActorYAxis,
+		ActorZAxis,
+		NumBuiltinConstants
+	} ConstantDef;
+
+	FName ConstantNames[] = {
+		"DeltaTime",
+		"Emitter position",
+		"Emitter age",
+		"Emitter X Axis",
+		"Emitter Y Axis",
+		"Emitter Z Axis",
+		""
+	};
+}
diff --git a/Engine/Source/Runtime/Engine/Classes/Engine/NiagaraScript.h b/Engine/Source/Runtime/Engine/Classes/Engine/NiagaraScript.h
--- a/Engine/Source/Runtime/Engine/Classes/Engine/NiagaraScript.h
+++ b/Engine/Source/Runtime/Engine/Classes/Engine/NiagaraScript.h
@@ -3,6 +3,8 @@
 #pragma once
 #include "NiagaraScript.generated.h"
 
+
+
 /** Runtime script for a Niagara system */
 UCLASS(MinimalAPI)
 class UNiagaraScript : public UObject
@@ -15,7 +17,7 @@ class UNiagaraScript : public UObject
 
 	/** The constant table required for execution. */
 	UPROPERTY()
-	TArray<float> ConstantTable;
+	TArray<FVector4> ConstantTable;
 
 	/** Attributes used by this script. */
 	TArray<FName> Attributes;
diff --git a/Engine/Source/Runtime/Engine/Private/NiagaraComponent.cpp b/Engine/Source/Runtime/Engine/Private/NiagaraComponent.cpp
--- a/Engine/Source/Runtime/Engine/Private/NiagaraComponent.cpp
+++ b/Engine/Source/Runtime/Engine/Private/NiagaraComponent.cpp
@@ -6,6 +6,7 @@
 #include "VectorVM.h"
 #include "ParticleHelper.h"
 #include "Particles/ParticleResources.h"
+#include "Engine/NiagaraConstants.h"
 
 DECLARE_CYCLE_STAT(TEXT("Tick"),STAT_NiagaraTick,STATGROUP_Niagara);
 DECLARE_CYCLE_STAT(TEXT("Simulate"),STAT_NiagaraSimulate,STATGROUP_Niagara);
@@ -15,6 +16,9 @@ DECLARE_CYCLE_STAT(TEXT("PreRenderView"),STAT_NiagaraPreRenderView,STATGROUP_Nia
 DECLARE_CYCLE_STAT(TEXT("Render"),STAT_NiagaraRender,STATGROUP_Niagara);
 DECLARE_DWORD_COUNTER_STAT(TEXT("NumParticles"),STAT_NiagaraNumParticles,STATGROUP_Niagara);
 
+
+DEFINE_LOG_CATEGORY_STATIC(LogNiagaraComponent, All, All);
+
 /** Struct used to pass dynamic data from game thread to render thread */
 struct FNiagaraDynamicData
 {
@@ -79,12 +83,13 @@ private:
 		UniformParameters.AxisLockUp = FVector4(0.0f, 0.0f, 0.0f, 0.0f);
 		UniformParameters.RotationScale = 1.0f;
 		UniformParameters.RotationBias = 0.0f;
-		UniformParameters.TangentSelector = FVector4(1.0f, 0.0f, 0.0f, 0.0f);
+		UniformParameters.TangentSelector = FVector4(0.0f, 0.0f, 0.0f, 1.0f);
 		UniformParameters.InvDeltaSeconds = 0.0f;
 		UniformParameters.SubImageSize = FVector4(1.0f,1.0f,1.0f,1.0f);
 		UniformParameters.NormalsType = 0;
 		UniformParameters.NormalsSphereCenter = FVector4(0.0f, 0.0f, 0.0f, 1.0f);
 		UniformParameters.NormalsCylinderUnitDirection = FVector4(0.0f, 0.0f, 1.0f, 0.0f);
+		UniformParameters.PivotOffset = FVector2D(0.0f, 0.0f);
 	}
 
 	virtual void OnActorPositionChanged() override
@@ -100,7 +105,7 @@ private:
 	virtual void PreRenderView(const FSceneViewFamily* ViewFamily, const uint32 VisibilityMap, int32 FrameNumber) override
 	{
 		SCOPE_CYCLE_COUNTER(STAT_NiagaraPreRenderView);
-
+		
 		check(DynamicData && DynamicData->VertexData.Num());
 
 		int32 SizeInBytes = DynamicData->VertexData.GetTypeSize() * DynamicData->VertexData.Num();
@@ -115,8 +120,7 @@ private:
 				sizeof(FParticleSpriteVertex),
 				true
 				);
-			VertexFactory.SetDynamicParameterBuffer(NULL,0,0,true);
-
+			VertexFactory.SetDynamicParameterBuffer(NULL, 0, 0, true);
 			// Compute the per-view uniform buffers.
 			const int32 NumViews = ViewFamily->Views.Num();
 			uint32 ViewBit = 1;
@@ -126,7 +130,7 @@ private:
 				if (VisibilityMap & ViewBit)
 				{
 					FParticleSpriteUniformParameters PerViewUniformParameters = UniformParameters;
-					PerViewUniformParameters.MacroUVParameters = FVector4(0.0f,0.0f,1.0f,1.0f);
+					PerViewUniformParameters.MacroUVParameters = FVector4(0.0f, 0.0f, 1.0f, 1.0f);
 					*SpriteViewUniformBufferPtr = FParticleSpriteUniformBufferRef::CreateUniformBufferImmediate(PerViewUniformParameters, UniformBuffer_SingleFrame);
 				}
 			}
@@ -160,6 +164,8 @@ private:
 		if (DynamicVertexAllocation.IsValid()
 			&& (bIsWireframe || !PDI->IsMaterialIgnored(MaterialRenderProxy, View->GetFeatureLevel())))
 		{
+			SCOPED_DRAW_EVENT(NiagaraRender, DEC_SCENE_ITEMS);
+
 			FMeshBatch MeshBatch;
 			MeshBatch.VertexFactory = &VertexFactory;
 			MeshBatch.CastShadow = CastsDynamicShadow();
@@ -167,6 +173,7 @@ private:
 			MeshBatch.ReverseCulling = IsLocalToWorldDeterminantNegative();
 			MeshBatch.Type = PT_TriangleList;
 			MeshBatch.DepthPriorityGroup = GetDepthPriorityGroup(View);
+			MeshBatch.LCI = NULL;
 			if (bIsWireframe)
 			{
 				MeshBatch.MaterialRenderProxy = UMaterial::GetDefaultMaterial(MD_Surface)->GetRenderProxy(IsSelected(),IsHovered());
@@ -205,7 +212,7 @@ private:
 		for(int32 i=0; i<DynamicData->VertexData.Num(); i++)
 		{
 			FParticleSpriteVertex& Vertex = DynamicData->VertexData[i];
-			PDI->DrawPoint(Vertex.Position, Vertex.Color, 3.f, SDPG_World);
+			PDI->DrawPoint(Vertex.Position, Vertex.Color, 1.0f, SDPG_World);
 		}
 		*/
 	}
@@ -259,29 +266,23 @@ private:
 	FMaterialRelevance MaterialRelevance;
 };
 
-namespace ENiagaraAttr
+namespace ENiagaraVectorAttr
 {
 	enum Type
 	{
-		PositionX,
-		PositionY,
-		PositionZ,
-		VelocityX,
-		VelocityY,
-		VelocityZ,
-		ColorR,
-		ColorG,
-		ColorB,
-		ColorA,
+		Position,
+		Velocity,
+		Color,
 		Rotation,
 		RelativeTime,
-		MaxAttributes
+		MaxVectorAttribs
 	};
 }
 
+
 /**
  * A niagara particle simulation.
- */
+ */n
 class FNiagaraSimulation
 {
 public:
@@ -289,6 +290,7 @@ public:
 		: SpawnRate(InComponent->SpawnRate)
 		, Component(*InComponent)
 		, UpdateScript(*InComponent->UpdateScript)
+		, SpawnScript(InComponent->SpawnScript)
 		, ConstantTable(InComponent->UpdateScript->ConstantTable)
 		, BufferIndex(0)
 		, NumVectorsPerAttribute(0)
@@ -297,9 +299,11 @@ public:
 		, CachedBounds(ForceInit)
 	{
 		check(InComponent->UpdateScript && UpdateScript.ByteCode.Num());
-		while (ConstantTable.Num() < 2)
+
+		// make room for constants
+		while (ConstantTable.Num() < 10)	
 		{
-			ConstantTable.Add(0);
+			ConstantTable.Add( FVector4(0.0f, 0.0f, 0.0f, 0.0f) );
 		}
 	}
 
@@ -328,11 +332,11 @@ public:
 		// The script updates relative time so we don't know yet which will die
 		// without simulation. Allocate for the worst case.
 		int32 MaxNewParticles = NumParticles + NumToSpawn;
-		NumVectorsPerAttribute = ((MaxNewParticles + 0x3) & ~0x3) >> 2;
+		NumVectorsPerAttribute = MaxNewParticles;
 		Particles.Reset(NumAttributes * NumVectorsPerAttribute);
 		Particles.AddUninitialized(NumAttributes * NumVectorsPerAttribute);
 
-		// Simualte particles forward by DeltaSeconds.
+		// Simulate particles forward by DeltaSeconds.
 		{
 			SCOPE_CYCLE_COUNTER(STAT_NiagaraSimulate);
 			UpdateParticles(
@@ -370,6 +374,17 @@ public:
 
 	FBox GetBounds() const { return CachedBounds; }
 
+	void AddConstant(FVector4 Value)
+	{
+		ConstantTable.Add(Value);
+	}
+
+	void SetConstant(int Idx, FVector4 Value)
+	{
+		check(Idx < ConstantTable.Num());
+		ConstantTable[Idx] = Value;
+	}
+
 private:
 	/** Temporary stuff for the prototype. */
 	float SpawnRate;
@@ -378,8 +393,10 @@ private:
 	const UNiagaraComponent& Component;
 	/** The particle update script. */
 	UNiagaraScript& UpdateScript;
+	/** The particle spawn script. */
+	UNiagaraScript *SpawnScript;
 	/** Local constant table. */
-	TArray<float> ConstantTable;
+	TArray<FVector4> ConstantTable;
 	/** 16-byte aligned set of attribute data used by vector VM for simulation */
 	TArray<FVector4> ParticleBuffers[2];
 	/** Current buffer. */
@@ -395,6 +412,20 @@ private:
 	/** Cached bounds. */
 	FBox CachedBounds;
 
+
+	/* builtin constants are setup in the tick function, but a script may have additional constants 
+	 derived from unconnected nodes, so we need to merge those into the simulation constant table */
+	void MergeScriptConstants(UNiagaraScript *Script)
+	{
+		// start at NumBuiltinConstants+1, because index 0 is always the fixed zero constant
+		for (int i = NiagaraConstants::NumBuiltinConstants + 1; i < Script->ConstantTable.Num(); i++)
+		{
+			if (ConstantTable.Num() <= i)
+				ConstantTable.Add(FVector4());
+			ConstantTable[i] = Script->ConstantTable[i];
+		}
+	}
+
 	/** Calc number to spawn */
 	int32 CalcNumToSpawn(float DeltaSeconds)
 	{
@@ -417,7 +448,7 @@ private:
 		VectorRegister* InputRegisters[VectorVM::MaxInputRegisters] = {0};
 		VectorRegister* OutputRegisters[VectorVM::MaxOutputRegisters] = {0};
 		const int32 NumAttr = UpdateScript.Attributes.Num();
-		const int32 NumVectors = ((NumParticles + 0x3) & ~0x3) >> 2;
+		const int32 NumVectors = NumParticles;
 
 		check(NumAttr < VectorVM::MaxInputRegisters);
 		check(NumAttr < VectorVM::MaxOutputRegisters);
@@ -429,9 +460,8 @@ private:
 			OutputRegisters[AttrIndex] = (VectorRegister*)(Particles + AttrIndex * NumVectorsPerAttribute);
 		}
 
-		// Setup constant table.
-		ConstantTable[0] = 0.0f;
-		ConstantTable[1] = DeltaSeconds;
+		// copy script specific constants into the constant table
+		MergeScriptConstants(&UpdateScript);
 
 		VectorVM::Exec(
 			UpdateScript.ByteCode.GetData(),
@@ -447,44 +477,65 @@ private:
 	/** Init new particles, and write over dead ones */
 	int32 SpawnAndKillParticles(
 		FVector4* Particles,
-		int32 NumParticles, 
-		int32 NumToSpawn, 
+		int32 NumParticles,
+		int32 NumToSpawn,
 		int32 NumVectorsPerAttribute
 		)
 	{
-		// Iterate over looking for dead particles and spawning new ones
-		int32 ParticleIndex = 0;
-		float* ParticleRelativeTimes = (float*)(Particles + ENiagaraAttr::RelativeTime * NumVectorsPerAttribute);
-		while (NumToSpawn > 0 && ParticleIndex < NumParticles)
-		{
-			if (ParticleRelativeTimes[ParticleIndex] > 1.0f)
-			{
-				// Particle is dead, spawn a new one.
-				SpawnParticleAtIndex(Particles, NumVectorsPerAttribute, ParticleIndex);
-				NumToSpawn--;
-			}
-			ParticleIndex++;
-		}
+		int32 OrigNumParticles = NumParticles;
+		FVector4 *NewParticles = Particles + NumParticles;
 
-		// Spawn any remaining particles.
-		while (NumToSpawn > 0)
+		// Spawn new Particles at the end of the buffer
+		for (int32 i = 0; i < NumToSpawn; i++)
 		{
-			SpawnParticleAtIndex(Particles, NumVectorsPerAttribute, ParticleIndex++);
-			NumToSpawn--;
+			SpawnParticleAtIndex(NewParticles, NumVectorsPerAttribute, i);
 			NumParticles++;
 		}
 
-		// Compact remaining particles.
-		while (ParticleIndex < NumParticles)
+		// run the spawn graph over all new particles
+		if (SpawnScript && SpawnScript->ByteCode.Num())
 		{
-			if (ParticleRelativeTimes[ParticleIndex] > 1.0f)
+			VectorRegister* InputRegisters[VectorVM::MaxInputRegisters] = { 0 };
+			VectorRegister* OutputRegisters[VectorVM::MaxOutputRegisters] = { 0 };
+			const int32 NumAttr = SpawnScript->Attributes.Num();
+			const int32 NumVectors = NumToSpawn;// ((NumParticles + 0x3) & ~0x3);
+
+			check(NumAttr < VectorVM::MaxInputRegisters);
+			check(NumAttr < VectorVM::MaxOutputRegisters);
+
+			// Setup input and output registers.
+			for (int32 AttrIndex = 0; AttrIndex < NumAttr; ++AttrIndex)
 			{
-				MoveParticleToIndex(Particles, NumVectorsPerAttribute, --NumParticles, ParticleIndex);
+				InputRegisters[AttrIndex] = (VectorRegister*)(NewParticles + AttrIndex * NumVectorsPerAttribute);
+				OutputRegisters[AttrIndex] = (VectorRegister*)(NewParticles + AttrIndex * NumVectorsPerAttribute);
 			}
-			else
+
+			// copy script specific constants into the constant table
+			MergeScriptConstants(SpawnScript);
+
+			VectorVM::Exec(
+				SpawnScript->ByteCode.GetData(),
+				InputRegisters,
+				NumAttr,
+				OutputRegisters,
+				NumAttr,
+				ConstantTable.GetData(),
+				NumVectors
+				);
+		}
+
+
+		// Iterate over looking for dead particles and move from the end of the list to the dead location, compacting in the process
+		int32 ParticleIndex = 0;
+		FVector4* ParticleRelativeTimes = (Particles + ENiagaraVectorAttr::RelativeTime * NumVectorsPerAttribute);
+		while (ParticleIndex < OrigNumParticles)
+		{
+			if (ParticleRelativeTimes[ParticleIndex].X > 1.0f)
 			{
-				ParticleIndex++;
+				// Particle is dead, move one from the end here.
+				MoveParticleToIndex(Particles, NumVectorsPerAttribute, --NumParticles, ParticleIndex);
 			}
+			ParticleIndex++;
 		}
 
 		return NumParticles;
@@ -494,37 +545,29 @@ private:
 	void SpawnParticleAtIndex(FVector4* Particles, int32 NumVectorsPerAttribute, int32 ParticleIndex)
 	{
 		FVector SpawnLocation = CachedComponentToWorld.GetLocation();
-		SpawnLocation.X += FMath::FRandRange(-50.f, 50.f);
-
-		float* Attr = (float*)Particles + ParticleIndex;
-		int32 NumFloatsPerAttribute = NumVectorsPerAttribute * 4;
-
-		*Attr = SpawnLocation.X; Attr += NumFloatsPerAttribute;
-		*Attr = SpawnLocation.Y; Attr += NumFloatsPerAttribute;
-		*Attr = SpawnLocation.Z; Attr += NumFloatsPerAttribute;
-		*Attr = 0.0f; Attr += NumFloatsPerAttribute;
-		*Attr = 0.0f; Attr += NumFloatsPerAttribute;
-		*Attr = 10.0f; Attr += NumFloatsPerAttribute;
-		*Attr = 1.0f; Attr += NumFloatsPerAttribute;
-		*Attr = 1.0f; Attr += NumFloatsPerAttribute;
-		*Attr = 1.0f; Attr += NumFloatsPerAttribute;
-		*Attr = 1.0f; Attr += NumFloatsPerAttribute;
-		*Attr = 0.0f; Attr += NumFloatsPerAttribute;
-		*Attr = 0.0f;
+		SpawnLocation.X += FMath::FRandRange(-20.f, 20.f);
+		SpawnLocation.Y += FMath::FRandRange(-20.f, 20.f);
+		SpawnLocation.Z += FMath::FRandRange(-20.f, 20.f);
+
+		FVector4* Attr = Particles + ParticleIndex;
+		Attr[NumVectorsPerAttribute*ENiagaraVectorAttr::Position] = SpawnLocation;
+		Attr[NumVectorsPerAttribute*ENiagaraVectorAttr::Color] = FVector4(0.0f, 0.0f, 0.0f, 0.0f);
+		Attr[NumVectorsPerAttribute*ENiagaraVectorAttr::RelativeTime] = FVector4(0.0f, 0.0f, 0.0f, 0.0f);
+		Attr[NumVectorsPerAttribute*ENiagaraVectorAttr::Rotation] = FVector4(0.0f, 0.0f, 0.0f, 0.0f);
+		Attr[NumVectorsPerAttribute*ENiagaraVectorAttr::Velocity] = FVector4(0.0f, 0.0f, 2.0f, 0.0f);
 	}
 
 	/** Util to move a particle */
 	void MoveParticleToIndex(FVector4* Particles, int32 NumVectorsPerAttribute, int32 SrcIndex, int32 DestIndex)
 	{
-		float* SrcPtr = (float*)Particles + SrcIndex;
-		float* DestPtr = (float*)Particles + DestIndex;
-		int32 NumFloatsPerAttribute = NumVectorsPerAttribute * 4;
+		FVector4 *SrcPtr = Particles + SrcIndex;
+		FVector4 *DestPtr = Particles + DestIndex;
 
-		for (int32 AttrIndex = 0; AttrIndex < ENiagaraAttr::MaxAttributes; AttrIndex++)
+		for (int32 AttrIndex = 0; AttrIndex < ENiagaraVectorAttr::MaxVectorAttribs; AttrIndex++)
 		{
 			*DestPtr = *SrcPtr;
-			DestPtr += NumFloatsPerAttribute;
-			SrcPtr += NumFloatsPerAttribute;
+			DestPtr += NumVectorsPerAttribute;
+			SrcPtr += NumVectorsPerAttribute;
 		}
 	}
 
@@ -534,27 +577,28 @@ private:
 		RenderData.Reset(NumParticles);
 		CachedBounds.Init();
 
+
 		float MaxSize = 0.0f;
-		float* Particles = (float*)ParticleBuffers[BufferIndex].GetTypedData();
-		int32 AttrStride = NumVectorsPerAttribute * 4;
+		FVector4* Particles = ParticleBuffers[BufferIndex].GetTypedData();
+		int32 AttrStride = NumVectorsPerAttribute;
 		for (int32 ParticleIndex = 0; ParticleIndex < NumParticles; ParticleIndex++)
 		{
-			float* Particle = Particles + ParticleIndex;
-		
-			FParticleSpriteVertex& NewVertex = *new(RenderData) FParticleSpriteVertex;
-			NewVertex.Position = FVector(Particle[AttrStride*ENiagaraAttr::PositionX],Particle[AttrStride*ENiagaraAttr::PositionY],Particle[AttrStride*ENiagaraAttr::PositionZ]);
-			//CachedBounds += NewVertex.Position;
-			NewVertex.Color = FLinearColor(Particle[AttrStride*ENiagaraAttr::ColorR],Particle[AttrStride*ENiagaraAttr::ColorG],Particle[AttrStride*ENiagaraAttr::ColorB],Particle[AttrStride*ENiagaraAttr::ColorA]);
+			FVector4* Particle = Particles + ParticleIndex;
+
+			FParticleSpriteVertex& NewVertex = *new(RenderData)FParticleSpriteVertex;
+			NewVertex.Position = Particle[AttrStride*ENiagaraVectorAttr::Position];
 			NewVertex.OldPosition = NewVertex.Position;
-			NewVertex.ParticleId = 0.f;
-			NewVertex.RelativeTime = Particle[AttrStride*ENiagaraAttr::RelativeTime];
-			NewVertex.Size = FVector2D(20.0f,20.0f);
-			//MaxSize = FMath::Max(MaxSize,NewVertex.Size.X);
-			//MaxSize = FMath::Max(MaxSize,NewVertex.Size.Y);
-			NewVertex.Rotation = Particle[AttrStride*ENiagaraAttr::Rotation];
+			CachedBounds += NewVertex.Position;
+			NewVertex.Color = FLinearColor(Particle[AttrStride*ENiagaraVectorAttr::Color]);
+			NewVertex.ParticleId = static_cast<float>(ParticleIndex) / NumParticles;
+			NewVertex.RelativeTime = Particle[AttrStride*ENiagaraVectorAttr::RelativeTime].X;
+			NewVertex.Size = FVector2D(1.0f, 1.0f);
+			MaxSize = FMath::Max(MaxSize, NewVertex.Size.X);
+			MaxSize = FMath::Max(MaxSize, NewVertex.Size.Y);
+			NewVertex.Rotation = Particle[AttrStride*ENiagaraVectorAttr::Rotation].X;
 			NewVertex.SubImageIndex = 0.f;
 		}
-		//CachedBounds.ExpandBy(MaxSize);
+		CachedBounds.ExpandBy(MaxSize);
 	}
 };
 
@@ -574,6 +618,16 @@ void UNiagaraComponent::TickComponent(float DeltaSeconds, enum ELevelTick TickTy
 {
 	if (Simulation)
 	{
+		EmitterAge += DeltaSeconds;
+
+		Simulation->SetConstant(0, FVector4(0.0f, 0.0f, 0.0f, 0.0f) );	// zero constant
+		Simulation->SetConstant(1, FVector4(DeltaSeconds, DeltaSeconds, DeltaSeconds, DeltaSeconds));
+		Simulation->SetConstant(2, FVector4(ComponentToWorld.GetTranslation()));
+		Simulation->SetConstant(3, FVector4(EmitterAge, EmitterAge, EmitterAge, EmitterAge));
+		Simulation->SetConstant(4, FVector4(ComponentToWorld.GetUnitAxis(EAxis::X)) );
+		Simulation->SetConstant(5, FVector4(ComponentToWorld.GetUnitAxis(EAxis::Y)) );
+		Simulation->SetConstant(6, FVector4(ComponentToWorld.GetUnitAxis(EAxis::Z)) );
+
 		Simulation->Tick(DeltaSeconds);
 		UpdateComponentToWorld();
 		MarkRenderDynamicDataDirty();
@@ -585,7 +639,7 @@ void UNiagaraComponent::OnRegister()
 	Super::OnRegister();
 
 	ensure(Simulation == NULL);
-	if (UpdateScript && UpdateScript->ByteCode.Num() && UpdateScript->Attributes.Num() == ENiagaraAttr::MaxAttributes)
+	if (UpdateScript && UpdateScript->ByteCode.Num() && UpdateScript->Attributes.Num() == ENiagaraVectorAttr::MaxVectorAttribs)
 	{
 		Simulation = new FNiagaraSimulation(this);
 	}
diff --git a/Engine/Source/Runtime/VectorVM/Private/VectorVM.cpp b/Engine/Source/Runtime/VectorVM/Private/VectorVM.cpp
--- a/Engine/Source/Runtime/VectorVM/Private/VectorVM.cpp
+++ b/Engine/Source/Runtime/VectorVM/Private/VectorVM.cpp
@@ -11,6 +11,14 @@ IMPLEMENT_MODULE(FDefaultModuleImpl, VectorVM);
 
 DEFINE_LOG_CATEGORY_STATIC(LogVectorVM, All, All);
 
+#define SRCOP_RRR 0x00
+#define SRCOP_RRC 0x01
+#define SRCOP_RCR 0x02
+#define SRCOP_RCC 0x03
+#define SRCOP_CRC 0x05
+#define SRCOP_CCR 0x06
+#define SRCOP_CCC 0x07
+
 /**
  * Context information passed around during VM execution.
  */
@@ -21,7 +29,7 @@ struct FVectorVMContext
 	/** Pointer to the table of vector register arrays. */
 	VectorRegister* RESTRICT * RESTRICT RegisterTable;
 	/** Pointer to the constant table. */
-	float const* RESTRICT ConstantTable;
+	FVector4 const* RESTRICT ConstantTable;
 	/** The number of vectors to process. */
 	int32 NumVectors;
 
@@ -29,7 +37,7 @@ struct FVectorVMContext
 	FVectorVMContext(
 		uint8 const* InCode,
 		VectorRegister** InRegisterTable,
-		float const* InConstantTable,
+		FVector4 const* InConstantTable,
 		int32 InNumVectors
 		)
 		: Code(InCode)
@@ -55,8 +63,13 @@ static FORCEINLINE VectorRegister* DecodeRegister(FVectorVMContext& Context)
 /** Decode a constant from the bytecode. */
 static FORCEINLINE VectorRegister DecodeConstant(FVectorVMContext& Context)
 {
-	const float* Float1 = &Context.ConstantTable[*Context.Code++];
-	return VectorLoadFloat1(Float1);
+	const FVector4* vec = &Context.ConstantTable[*Context.Code++];
+	return VectorLoad(vec);
+}
+
+static FORCEINLINE uint8 DecodeSrcOperandTypes(FVectorVMContext& Context)
+{
+	return *Context.Code++;
 }
 
 /** Base class for vector kernels with one dest and one src operand. */
@@ -66,15 +79,30 @@ struct TUnaryVectorKernel
 	static void Exec(FVectorVMContext& Context)
 	{
 		VectorRegister* RESTRICT Dst = DecodeRegister(Context);
-		VectorRegister* Src0 = DecodeRegister(Context);
+		uint32 SrcOpTypes = DecodeSrcOperandTypes(Context);
 		int32 NumVectors = Context.NumVectors;
-		for (int32 i = 0; i < NumVectors; ++i)
+
+		if (SrcOpTypes == SRCOP_RRR)
 		{
-			Kernel::DoKernel(Dst++,*Src0++);
+			VectorRegister* Src0 = DecodeRegister(Context);
+			for (int32 i = 0; i < NumVectors; ++i)
+			{
+				Kernel::DoKernel(Dst++,*Src0++);
+			}
 		}
+		else if (SrcOpTypes == SRCOP_RRC)
+		{
+			VectorRegister Src0 = DecodeConstant(Context);
+			for (int32 i = 0; i < NumVectors; ++i)
+			{
+				Kernel::DoKernel(Dst++, Src0);
+			}
+		}
+
 	}
 };
 
+
 /** Base class for vector kernels with one dest and two src operands. */
 template <typename Kernel>
 struct TBinaryVectorKernel
@@ -82,155 +110,117 @@ struct TBinaryVectorKernel
 	static void Exec(FVectorVMContext& Context)
 	{
 		VectorRegister* RESTRICT Dst = DecodeRegister(Context);
-		VectorRegister* Src0 = DecodeRegister(Context);
-		VectorRegister* Src1 = DecodeRegister(Context);
+		uint32 SrcOpTypes = DecodeSrcOperandTypes(Context);
+
 		int32 NumVectors = Context.NumVectors;
-		for (int32 i = 0; i < NumVectors; ++i)
+		if (SrcOpTypes == SRCOP_RRR)
 		{
-			Kernel::DoKernel(Dst++,*Src0++,*Src1++);
+			VectorRegister* Src0 = DecodeRegister(Context);
+			VectorRegister* Src1 = DecodeRegister(Context);
+			for (int32 i = 0; i < NumVectors; ++i)
+			{
+				Kernel::DoKernel(Dst++, *Src0++, *Src1++);
+			}
 		}
-	}
-};
-
-/** Base class for vector kernels with one dest and two src operands, one of which is constant. */
-template <typename Kernel>
-struct TBinaryVectorKernelWithConstant
-{
-	static void Exec(FVectorVMContext& Context)
-	{
-		VectorRegister* RESTRICT Dst = DecodeRegister(Context);
-		VectorRegister* Src0 = DecodeRegister(Context);
-		VectorRegister Src1 = DecodeConstant(Context);
-		int32 NumVectors = Context.NumVectors;
-		for (int32 i = 0; i < NumVectors; ++i)
+		else if (SrcOpTypes == SRCOP_RRC)
 		{
-			Kernel::DoKernel(Dst++,*Src0++,Src1);
+			VectorRegister Src0 = DecodeConstant(Context);
+			VectorRegister *Src1 = DecodeRegister(Context);
+			for (int32 i = 0; i < NumVectors; ++i)
+			{
+				Kernel::DoKernel(Dst++, Src0, *Src1++);
+			}
 		}
-	}
-};
-
-/** Base class for vector kernels with one dest and three src operands. */
-template <typename Kernel>
-struct TTrinaryVectorKernel
-{
-	static void Exec(FVectorVMContext& Context)
-	{
-		VectorRegister* RESTRICT Dst = DecodeRegister(Context);
-		VectorRegister* Src0 = DecodeRegister(Context);
-		VectorRegister* Src1 = DecodeRegister(Context);
-		VectorRegister* Src2 = DecodeRegister(Context);
-		int32 NumVectors = Context.NumVectors;
-		for (int32 i = 0; i < NumVectors; ++i)
+		else if (SrcOpTypes == SRCOP_RCR)
 		{
-			Kernel::DoKernel(Dst++,*Src0++,*Src1++,*Src2++);
+			VectorRegister *Src0 = DecodeRegister(Context);
+			VectorRegister Src1 = DecodeConstant(Context);
+			for (int32 i = 0; i < NumVectors; ++i)
+			{
+				Kernel::DoKernel(Dst++, *Src0++, Src1);
+			}
 		}
-	}
-};
-
-/** Base class for vector kernels with one dest and three src operands. */
-template <typename Kernel>
-struct TTrinaryVectorKernelIRR
-{
-	static void Exec(FVectorVMContext& Context)
-	{
-		VectorRegister* RESTRICT Dst = DecodeRegister(Context);
-		VectorRegister Src0 = DecodeConstant(Context);
-		VectorRegister* Src1 = DecodeRegister(Context);
-		VectorRegister* Src2 = DecodeRegister(Context);
-		int32 NumVectors = Context.NumVectors;
-		for (int32 i = 0; i < NumVectors; ++i)
+		else if (SrcOpTypes == SRCOP_RCC)
 		{
-			Kernel::DoKernel(Dst++,Src0,*Src1++,*Src2++);
+			VectorRegister Src0 = DecodeConstant(Context);
+			VectorRegister Src1 = DecodeConstant(Context);
+			for (int32 i = 0; i < NumVectors; ++i)
+			{
+				Kernel::DoKernel(Dst++, Src0, Src1);
+			}
 		}
 	}
 };
 
 /** Base class for vector kernels with one dest and three src operands. */
 template <typename Kernel>
-struct TTrinaryVectorKernelRIR
+struct TTrinaryVectorKernel
 {
 	static void Exec(FVectorVMContext& Context)
 	{
 		VectorRegister* RESTRICT Dst = DecodeRegister(Context);
-		VectorRegister* Src0 = DecodeRegister(Context);
-		VectorRegister Src1 = DecodeConstant(Context);
-		VectorRegister* Src2 = DecodeRegister(Context);
+		uint32 SrcOpTypes = DecodeSrcOperandTypes(Context);
+
 		int32 NumVectors = Context.NumVectors;
-		for (int32 i = 0; i < NumVectors; ++i)
+		if (SrcOpTypes == SRCOP_RRR)
 		{
-			Kernel::DoKernel(Dst++,*Src0++,Src1,*Src2++);
+			VectorRegister* Src0 = DecodeRegister(Context);
+			VectorRegister* Src1 = DecodeRegister(Context);
+			VectorRegister* Src2 = DecodeRegister(Context);
+			for (int32 i = 0; i < NumVectors; ++i)
+			{
+				Kernel::DoKernel(Dst++, *Src0++, *Src1++, *Src2++);
+			}
 		}
-	}
-};
-
-/** Base class for vector kernels with one dest and three src operands. */
-template <typename Kernel>
-struct TTrinaryVectorKernelRRI
-{
-	static void Exec(FVectorVMContext& Context)
-	{
-		VectorRegister* RESTRICT Dst = DecodeRegister(Context);
-		VectorRegister* Src0 = DecodeRegister(Context);
-		VectorRegister* Src1 = DecodeRegister(Context);
-		VectorRegister Src2 = DecodeConstant(Context);
-		int32 NumVectors = Context.NumVectors;
-		for (int32 i = 0; i < NumVectors; ++i)
+		else if (SrcOpTypes == SRCOP_RRC)
 		{
-			Kernel::DoKernel(Dst++,*Src0++,*Src1++,Src2);
+			VectorRegister Src0 = DecodeConstant(Context);
+			VectorRegister* Src1 = DecodeRegister(Context);
+			VectorRegister* Src2 = DecodeRegister(Context);
+			for (int32 i = 0; i < NumVectors; ++i)
+			{
+				Kernel::DoKernel(Dst++, Src0, *Src1++, *Src2++);
+			}
 		}
-	}
-};
-
-/** Base class for vector kernels with one dest and three src operands. */
-template <typename Kernel>
-struct TTrinaryVectorKernelRII
-{
-	static void Exec(FVectorVMContext& Context)
-	{
-		VectorRegister* RESTRICT Dst = DecodeRegister(Context);
-		VectorRegister* Src0 = DecodeRegister(Context);
-		VectorRegister Src1 = DecodeConstant(Context);
-		VectorRegister Src2 = DecodeConstant(Context);
-		int32 NumVectors = Context.NumVectors;
-		for (int32 i = 0; i < NumVectors; ++i)
+		else if (SrcOpTypes == SRCOP_CCC)
 		{
-			Kernel::DoKernel(Dst++,*Src0++,Src1,Src2);
+			VectorRegister Src0 = DecodeConstant(Context);
+			VectorRegister Src1 = DecodeConstant(Context);
+			VectorRegister Src2 = DecodeConstant(Context);
+			for (int32 i = 0; i < NumVectors; ++i)
+			{
+				Kernel::DoKernel(Dst++, Src0, Src1, Src2);
+			}
 		}
-	}
-};
-
-/** Base class for vector kernels with one dest and three src operands. */
-template <typename Kernel>
-struct TTrinaryVectorKernelIIR
-{
-	static void Exec(FVectorVMContext& Context)
-	{
-		VectorRegister* RESTRICT Dst = DecodeRegister(Context);
-		VectorRegister Src0 = DecodeConstant(Context);
-		VectorRegister Src1 = DecodeConstant(Context);
-		VectorRegister* Src2 = DecodeRegister(Context);
-		int32 NumVectors = Context.NumVectors;
-		for (int32 i = 0; i < NumVectors; ++i)
+		else if (SrcOpTypes == SRCOP_CCR)
 		{
-			Kernel::DoKernel(Dst++,Src0,Src1,*Src2++);
+			VectorRegister *Src0 = DecodeRegister(Context);
+			VectorRegister Src1 = DecodeConstant(Context);
+			VectorRegister Src2 = DecodeConstant(Context);
+			for (int32 i = 0; i < NumVectors; ++i)
+			{
+				Kernel::DoKernel(Dst++, *Src0++, Src1, Src2);
+			}
 		}
-	}
-};
-
-/** Base class for vector kernels with one dest and three src operands. */
-template <typename Kernel>
-struct TTrinaryVectorKernelIII
-{
-	static void Exec(FVectorVMContext& Context)
-	{
-		VectorRegister* RESTRICT Dst = DecodeRegister(Context);
-		VectorRegister Src0 = DecodeConstant(Context);
-		VectorRegister Src1 = DecodeConstant(Context);
-		VectorRegister Src2 = DecodeConstant(Context);
-		int32 NumVectors = Context.NumVectors;
-		for (int32 i = 0; i < NumVectors; ++i)
+		else if (SrcOpTypes == SRCOP_RCC)
+		{
+			VectorRegister Src0 = DecodeConstant(Context);
+			VectorRegister Src1 = DecodeConstant(Context);
+			VectorRegister *Src2 = DecodeRegister(Context);
+			for (int32 i = 0; i < NumVectors; ++i)
+			{
+				Kernel::DoKernel(Dst++, Src0, Src1, *Src2++);
+			}
+		}
+		else if (SrcOpTypes == SRCOP_RCR)
 		{
-			Kernel::DoKernel(Dst++,Src0,Src1,Src2);
+			VectorRegister *Src0 = DecodeRegister(Context);
+			VectorRegister Src1 = DecodeConstant(Context);
+			VectorRegister *Src2 = DecodeRegister(Context);
+			for (int32 i = 0; i < NumVectors; ++i)
+			{
+				Kernel::DoKernel(Dst++, *Src0++, Src1, *Src2++);
+			}
 		}
 	}
 };
@@ -246,7 +236,6 @@ struct FVectorKernelAdd : public TBinaryVectorKernel<FVectorKernelAdd>
 		*Dst = VectorAdd(Src0, Src1);
 	}
 };
-struct FVectorKernelAddi : public TBinaryVectorKernelWithConstant<FVectorKernelAdd> {};
 
 struct FVectorKernelSub : public TBinaryVectorKernel<FVectorKernelSub>
 {
@@ -255,7 +244,6 @@ struct FVectorKernelSub : public TBinaryVectorKernel<FVectorKernelSub>
 		*Dst = VectorSubtract(Src0, Src1);
 	}
 };
-struct FVectorKernelSubi : public TBinaryVectorKernelWithConstant<FVectorKernelSub> {};
 
 struct FVectorKernelMul : public TBinaryVectorKernel<FVectorKernelMul>
 {
@@ -264,7 +252,6 @@ struct FVectorKernelMul : public TBinaryVectorKernel<FVectorKernelMul>
 		*Dst = VectorMultiply(Src0, Src1);
 	}
 };
-struct FVectorKernelMuli : public TBinaryVectorKernelWithConstant<FVectorKernelMul> {};
 
 struct FVectorKernelMad : public TTrinaryVectorKernel<FVectorKernelMad>
 {
@@ -273,25 +260,17 @@ struct FVectorKernelMad : public TTrinaryVectorKernel<FVectorKernelMad>
 		*Dst = VectorMultiplyAdd(Src0, Src1, Src2);
 	}
 };
-struct FVectorKernelMadrri : public TTrinaryVectorKernelRRI<FVectorKernelMad> {};
-struct FVectorKernelMadrir : public TTrinaryVectorKernelRIR<FVectorKernelMad> {};
-struct FVectorKernelMadrii : public TTrinaryVectorKernelRII<FVectorKernelMad> {};
-struct FVectorKernelMadiir : public TTrinaryVectorKernelIIR<FVectorKernelMad> {};
-struct FVectorKernelMadiii : public TTrinaryVectorKernelIII<FVectorKernelMad> {};
 
 struct FVectorKernelLerp : public TTrinaryVectorKernel<FVectorKernelLerp>
 {
 	static void FORCEINLINE DoKernel(VectorRegister* RESTRICT Dst,VectorRegister Src0,VectorRegister Src1,VectorRegister Src2)
 	{
-		const VectorRegister OneMinusAlpha = VectorNegate(Src2);
+		const VectorRegister One = MakeVectorRegister(1.0f, 1.0f, 1.0f, 1.0f);
+		const VectorRegister OneMinusAlpha = VectorSubtract(One, Src2);
 		const VectorRegister Tmp = VectorMultiply(Src0, OneMinusAlpha);
 		*Dst = VectorMultiplyAdd(Src1, Src2, Tmp);
 	}
 };
-struct FVectorKernelLerpirr : public TTrinaryVectorKernelIRR<FVectorKernelLerp> {};
-struct FVectorKernelLerprir : public TTrinaryVectorKernelRIR<FVectorKernelLerp> {};
-struct FVectorKernelLerprri : public TTrinaryVectorKernelRRI<FVectorKernelLerp> {};
-struct FVectorKernelLerprii : public TTrinaryVectorKernelRII<FVectorKernelLerp> {};
 
 struct FVectorKernelRcp : public TUnaryVectorKernel<FVectorKernelRcp>
 {
@@ -347,9 +326,70 @@ struct FVectorKernelClamp : public TTrinaryVectorKernel<FVectorKernelClamp>
 		*Dst = VectorMin(Tmp, Src2);
 	}
 };
-struct FVectorKernelClampir : public TTrinaryVectorKernelRIR<FVectorKernelClamp> {};
-struct FVectorKernelClampri : public TTrinaryVectorKernelRRI<FVectorKernelClamp> {};
-struct FVectorKernelClampii : public TTrinaryVectorKernelRII<FVectorKernelClamp> {};
+
+
+struct FVectorKernelSin : public TUnaryVectorKernel<FVectorKernelSin>
+{
+	static void FORCEINLINE DoKernel(VectorRegister* RESTRICT Dst, VectorRegister Src0)
+	{
+		float v = VectorGetComponent(Src0, 0);
+		float sn = FMath::Sin(v*3.14f);		// [0;1] takes us through half a period
+		*Dst = MakeVectorRegister(sn, sn, sn, sn);
+	}
+};
+
+struct FVectorKernelDot : public TBinaryVectorKernel<FVectorKernelDot>
+{
+	static void FORCEINLINE DoKernel(VectorRegister* RESTRICT Dst, VectorRegister Src0, VectorRegister Src1)
+	{
+		*Dst = VectorDot4(Src0, Src1);
+	}
+};
+
+
+struct FVectorKernelLength : public TUnaryVectorKernel<FVectorKernelLength>
+{
+	static void FORCEINLINE DoKernel(VectorRegister* RESTRICT Dst, VectorRegister Src0)
+	{
+		VectorRegister Temp = VectorDot4(Src0, Src0);
+		float const* FloatSrc = reinterpret_cast<float const*>(&Temp);
+		float SDot = FMath::Sqrt(FloatSrc[0]);
+		*Dst = MakeVectorRegister(SDot, SDot, SDot, SDot);
+	}
+};
+
+
+struct FVectorKernelCross : public TBinaryVectorKernel<FVectorKernelCross>
+{
+	static void FORCEINLINE DoKernel(VectorRegister* RESTRICT Dst, VectorRegister Src0, VectorRegister Src1)
+	{
+		*Dst = VectorCross(Src0, Src1);
+	}
+};
+
+struct FVectorKernelNormalize : public TUnaryVectorKernel<FVectorKernelNormalize>
+{
+	static void FORCEINLINE DoKernel(VectorRegister* RESTRICT Dst, VectorRegister Src0)
+	{
+		*Dst = VectorNormalize(Src0);
+	}
+};
+
+struct FVectorKernelRandom : public TUnaryVectorKernel<FVectorKernelRandom>
+{
+	static void FORCEINLINE DoKernel(VectorRegister* RESTRICT Dst, VectorRegister Src0)
+	{
+		const float rm = RAND_MAX;
+		VectorRegister Result = MakeVectorRegister(static_cast<float>(FMath::Rand()) / rm,
+			static_cast<float>(FMath::Rand()) / rm,
+			static_cast<float>(FMath::Rand()) / rm,
+			static_cast<float>(FMath::Rand()) / rm);
+		*Dst = VectorMultiply(Result, Src0);
+		
+	}
+};
+
+
 
 struct FVectorKernelMin : public TBinaryVectorKernel<FVectorKernelMin>
 {
@@ -358,7 +398,6 @@ struct FVectorKernelMin : public TBinaryVectorKernel<FVectorKernelMin>
 		*Dst = VectorMin(Src0, Src1);
 	}
 };
-struct FVectorKernelMini : public TBinaryVectorKernelWithConstant<FVectorKernelMin> {};
 
 struct FVectorKernelMax : public TBinaryVectorKernel<FVectorKernelMax>
 {
@@ -367,7 +406,6 @@ struct FVectorKernelMax : public TBinaryVectorKernel<FVectorKernelMax>
 		*Dst = VectorMax(Src0, Src1);
 	}
 };
-struct FVectorKernelMaxi : public TBinaryVectorKernelWithConstant<FVectorKernelMax> {};
 
 struct FVectorKernelPow : public TBinaryVectorKernel<FVectorKernelPow>
 {
@@ -376,7 +414,9 @@ struct FVectorKernelPow : public TBinaryVectorKernel<FVectorKernelPow>
 		*Dst = VectorPow(Src0, Src1);
 	}
 };
-struct FVectorKernelPowi : public TBinaryVectorKernelWithConstant<FVectorKernelPow> {};
+
+
+
 
 void VectorVM::Exec(
 	uint8 const* Code,
@@ -384,7 +424,7 @@ void VectorVM::Exec(
 	int32 NumInputRegisters,
 	VectorRegister** OutputRegisters,
 	int32 NumOutputRegisters,
-	float const* ConstantTable,
+	FVector4 const* ConstantTable,
 	int32 NumVectors
 	)
 {
@@ -424,37 +464,25 @@ void VectorVM::Exec(
 			{
 			// Dispatch kernel ops.
 			case EOp::add: FVectorKernelAdd::Exec(Context); break;
-			case EOp::addi: FVectorKernelAddi::Exec(Context); break;
 			case EOp::sub: FVectorKernelSub::Exec(Context); break;
-			case EOp::subi: FVectorKernelSubi::Exec(Context); break;
 			case EOp::mul: FVectorKernelMul::Exec(Context); break;
-			case EOp::muli: FVectorKernelMuli::Exec(Context); break;
 			case EOp::mad: FVectorKernelMad::Exec(Context); break;
-			case EOp::madrri: FVectorKernelMadrri::Exec(Context); break;
-			case EOp::madrir: FVectorKernelMadrir::Exec(Context); break;
-			case EOp::madrii: FVectorKernelMadrii::Exec(Context); break;
-			case EOp::madiir: FVectorKernelMadiir::Exec(Context); break;
-			case EOp::madiii: FVectorKernelMadiii::Exec(Context); break;
 			case EOp::lerp: FVectorKernelLerp::Exec(Context); break;
-			case EOp::lerpirr: FVectorKernelLerpirr::Exec(Context); break;
-			case EOp::lerprir: FVectorKernelLerprir::Exec(Context); break;
-			case EOp::lerprri: FVectorKernelLerprri::Exec(Context); break;
-			case EOp::lerprii: FVectorKernelLerprii::Exec(Context); break;
 			case EOp::rcp: FVectorKernelRcp::Exec(Context); break;
 			case EOp::rsq: FVectorKernelRsq::Exec(Context); break;
 			case EOp::sqrt: FVectorKernelSqrt::Exec(Context); break;
 			case EOp::neg: FVectorKernelNeg::Exec(Context); break;
 			case EOp::abs: FVectorKernelAbs::Exec(Context); break;
 			case EOp::clamp: FVectorKernelClamp::Exec(Context); break;
-			case EOp::clampir: FVectorKernelClampir::Exec(Context); break;
-			case EOp::clampri: FVectorKernelClampri::Exec(Context); break;
-			case EOp::clampii: FVectorKernelClampii::Exec(Context); break;
 			case EOp::min: FVectorKernelMin::Exec(Context); break;
-			case EOp::mini: FVectorKernelMini::Exec(Context); break;
 			case EOp::max: FVectorKernelMax::Exec(Context); break;
-			case EOp::maxi: FVectorKernelMaxi::Exec(Context); break;
 			case EOp::pow: FVectorKernelPow::Exec(Context); break;
-			case EOp::powi: FVectorKernelPowi::Exec(Context); break;
+			case EOp::sin: FVectorKernelSin::Exec(Context); break;
+			case EOp::dot: FVectorKernelDot::Exec(Context); break;
+			case EOp::length: FVectorKernelLength::Exec(Context); break;
+			case EOp::cross: FVectorKernelCross::Exec(Context); break;
+			case EOp::normalize: FVectorKernelNormalize::Exec(Context); break;
+			case EOp::random: FVectorKernelRandom::Exec(Context); break;
 
 			// Execution always terminates with a "done" opcode.
 			case EOp::done:
@@ -475,73 +503,85 @@ namespace VectorVM
 {
 	static FVectorVMOpInfo GOpInfo[] =
 	{
-		FVectorVMOpInfo(EOp::done,EOpFlags::None,EOpSrc::Invalid,EOpSrc::Invalid,EOpSrc::Invalid,TEXT("done")),
+		FVectorVMOpInfo(EOp::done, EOpFlags::None, EOpSrc::Invalid, EOpSrc::Invalid, EOpSrc::Invalid, TEXT("done")),
 
-		FVectorVMOpInfo(EOp::add,EOpFlags::Implemented|EOpFlags::Commutative,EOpSrc::Register,EOpSrc::Register,EOpSrc::Invalid,TEXT("add")),
-		FVectorVMOpInfo(EOp::add,EOpFlags::None,EOpSrc::Register,EOpSrc::Const,EOpSrc::Invalid,TEXT("addi")),
+		FVectorVMOpInfo(EOp::add, EOpFlags::Implemented | EOpFlags::Commutative, EOpSrc::Register, EOpSrc::Register, EOpSrc::Invalid, TEXT("Add")),
+		FVectorVMOpInfo(EOp::add, EOpFlags::None, EOpSrc::Register, EOpSrc::Const, EOpSrc::Invalid, TEXT("addi")),
 
-		FVectorVMOpInfo(EOp::sub,EOpFlags::Implemented,EOpSrc::Register,EOpSrc::Register,EOpSrc::Invalid,TEXT("sub")),
-		FVectorVMOpInfo(EOp::sub,EOpFlags::None,EOpSrc::Register,EOpSrc::Const,EOpSrc::Invalid,TEXT("subi")),
+		FVectorVMOpInfo(EOp::sub, EOpFlags::Implemented, EOpSrc::Register, EOpSrc::Register, EOpSrc::Invalid, TEXT("Sub")),
+		FVectorVMOpInfo(EOp::sub, EOpFlags::None, EOpSrc::Register, EOpSrc::Const, EOpSrc::Invalid, TEXT("subi")),
 
-		FVectorVMOpInfo(EOp::mul,EOpFlags::Implemented|EOpFlags::Commutative,EOpSrc::Register,EOpSrc::Register,EOpSrc::Invalid,TEXT("mul")),
+		FVectorVMOpInfo(EOp::mul, EOpFlags::Implemented | EOpFlags::Commutative, EOpSrc::Register, EOpSrc::Register, EOpSrc::Invalid, TEXT("Multiply")),
 		FVectorVMOpInfo(EOp::mul,EOpFlags::None,EOpSrc::Register,EOpSrc::Const,EOpSrc::Invalid,TEXT("muli")),
 
-		FVectorVMOpInfo(EOp::mad,EOpFlags::Implemented|EOpFlags::Commutative,EOpSrc::Register,EOpSrc::Register,EOpSrc::Register,TEXT("mad")),
+		FVectorVMOpInfo(EOp::mad,EOpFlags::Implemented|EOpFlags::Commutative,EOpSrc::Register,EOpSrc::Register,EOpSrc::Register,TEXT("Multply-Add")),
 		FVectorVMOpInfo(EOp::mad,EOpFlags::None,EOpSrc::Register,EOpSrc::Register,EOpSrc::Const,TEXT("madrri")),
 		FVectorVMOpInfo(EOp::mad,EOpFlags::None,EOpSrc::Register,EOpSrc::Const,EOpSrc::Register,TEXT("madrir")),
 		FVectorVMOpInfo(EOp::mad,EOpFlags::None,EOpSrc::Register,EOpSrc::Const,EOpSrc::Const,TEXT("madrii")),
 		FVectorVMOpInfo(EOp::mad,EOpFlags::None,EOpSrc::Const,EOpSrc::Const,EOpSrc::Register,TEXT("madiir")),
 		FVectorVMOpInfo(EOp::mad,EOpFlags::None,EOpSrc::Const,EOpSrc::Const,EOpSrc::Const,TEXT("madiii")),
 
-		FVectorVMOpInfo(EOp::lerp,EOpFlags::Implemented,EOpSrc::Register,EOpSrc::Register,EOpSrc::Register,TEXT("lerp")),
+		FVectorVMOpInfo(EOp::lerp,EOpFlags::Implemented,EOpSrc::Register,EOpSrc::Register,EOpSrc::Register,TEXT("Lerp")),
 		FVectorVMOpInfo(EOp::lerp,EOpFlags::None,EOpSrc::Const,EOpSrc::Register,EOpSrc::Register,TEXT("lerpirr")),
 		FVectorVMOpInfo(EOp::lerp,EOpFlags::None,EOpSrc::Register,EOpSrc::Const,EOpSrc::Register,TEXT("lerprir")), 
 		FVectorVMOpInfo(EOp::lerp,EOpFlags::None,EOpSrc::Register,EOpSrc::Register,EOpSrc::Const,TEXT("lerprri")),
-		FVectorVMOpInfo(EOp::lerp,EOpFlags::None,EOpSrc::Register,EOpSrc::Const,EOpSrc::Const,TEXT("lerprii")),
-
-		FVectorVMOpInfo(EOp::rcp,EOpFlags::Implemented,EOpSrc::Register,EOpSrc::Invalid,EOpSrc::Invalid,TEXT("rcp")),
-		FVectorVMOpInfo(EOp::rsq,EOpFlags::Implemented,EOpSrc::Register,EOpSrc::Invalid,EOpSrc::Invalid,TEXT("rsq")),
-		FVectorVMOpInfo(EOp::sqrt,EOpFlags::Implemented,EOpSrc::Register,EOpSrc::Invalid,EOpSrc::Invalid,TEXT("sqrt")),
-		FVectorVMOpInfo(EOp::neg,EOpFlags::Implemented,EOpSrc::Register,EOpSrc::Invalid,EOpSrc::Invalid,TEXT("neg")),
-		FVectorVMOpInfo(EOp::abs,EOpFlags::Implemented,EOpSrc::Register,EOpSrc::Invalid,EOpSrc::Invalid,TEXT("abs")),
-		FVectorVMOpInfo(EOp::exp,EOpFlags::None,EOpSrc::Register,EOpSrc::Invalid,EOpSrc::Invalid,TEXT("exp")),
-		FVectorVMOpInfo(EOp::exp2,EOpFlags::None,EOpSrc::Register,EOpSrc::Invalid,EOpSrc::Invalid,TEXT("exp2")),
-		FVectorVMOpInfo(EOp::log,EOpFlags::None,EOpSrc::Register,EOpSrc::Invalid,EOpSrc::Invalid,TEXT("log")),
-		FVectorVMOpInfo(EOp::log2,EOpFlags::None,EOpSrc::Register,EOpSrc::Invalid,EOpSrc::Invalid,TEXT("log2")),
-		FVectorVMOpInfo(EOp::sin,EOpFlags::None,EOpSrc::Register,EOpSrc::Invalid,EOpSrc::Invalid,TEXT("sin")),
-		FVectorVMOpInfo(EOp::cos,EOpFlags::None,EOpSrc::Register,EOpSrc::Invalid,EOpSrc::Invalid,TEXT("cos")),
-		FVectorVMOpInfo(EOp::tan,EOpFlags::None,EOpSrc::Register,EOpSrc::Invalid,EOpSrc::Invalid,TEXT("tan")),
-		FVectorVMOpInfo(EOp::asin,EOpFlags::None,EOpSrc::Register,EOpSrc::Invalid,EOpSrc::Invalid,TEXT("asin")),
-		FVectorVMOpInfo(EOp::acos,EOpFlags::None,EOpSrc::Register,EOpSrc::Invalid,EOpSrc::Invalid,TEXT("acos")),
-		FVectorVMOpInfo(EOp::atan,EOpFlags::None,EOpSrc::Register,EOpSrc::Invalid,EOpSrc::Invalid,TEXT("atan")),
-		FVectorVMOpInfo(EOp::atan2,EOpFlags::None,EOpSrc::Register,EOpSrc::Register,EOpSrc::Invalid,TEXT("atan2")),
-		FVectorVMOpInfo(EOp::ceil,EOpFlags::None,EOpSrc::Register,EOpSrc::Invalid,EOpSrc::Invalid,TEXT("ceil")),
-		FVectorVMOpInfo(EOp::floor,EOpFlags::None,EOpSrc::Register,EOpSrc::Invalid,EOpSrc::Invalid,TEXT("floor")),
-		FVectorVMOpInfo(EOp::fmod,EOpFlags::None,EOpSrc::Register,EOpSrc::Register,EOpSrc::Invalid,TEXT("fmod")),
-		FVectorVMOpInfo(EOp::frac,EOpFlags::None,EOpSrc::Register,EOpSrc::Invalid,EOpSrc::Invalid,TEXT("frac")),
-		FVectorVMOpInfo(EOp::trunc,EOpFlags::None,EOpSrc::Register,EOpSrc::Invalid,EOpSrc::Invalid,TEXT("trunc")),
-
-		FVectorVMOpInfo(EOp::clamp,EOpFlags::Implemented,EOpSrc::Register,EOpSrc::Register,EOpSrc::Register,TEXT("clamp")),
+		FVectorVMOpInfo(EOp::lerp,EOpFlags::None,EOpSrc::Const,EOpSrc::Const,EOpSrc::Register,TEXT("lerpiir")),
+
+		FVectorVMOpInfo(EOp::rcp,EOpFlags::Implemented,EOpSrc::Register,EOpSrc::Invalid,EOpSrc::Invalid,TEXT("Reciprocal")),
+		FVectorVMOpInfo(EOp::rsq,EOpFlags::Implemented,EOpSrc::Register,EOpSrc::Invalid,EOpSrc::Invalid,TEXT("Reciprocal Sqrt")),
+		FVectorVMOpInfo(EOp::sqrt,EOpFlags::Implemented,EOpSrc::Register,EOpSrc::Invalid,EOpSrc::Invalid,TEXT("Sqrt")),
+		FVectorVMOpInfo(EOp::neg,EOpFlags::Implemented,EOpSrc::Register,EOpSrc::Invalid,EOpSrc::Invalid,TEXT("Negate")),
+		FVectorVMOpInfo(EOp::abs,EOpFlags::Implemented,EOpSrc::Register,EOpSrc::Invalid,EOpSrc::Invalid,TEXT("Absolute")),
+		FVectorVMOpInfo(EOp::exp,EOpFlags::None,EOpSrc::Register,EOpSrc::Invalid,EOpSrc::Invalid,TEXT("Exp")),
+		FVectorVMOpInfo(EOp::exp2,EOpFlags::None,EOpSrc::Register,EOpSrc::Invalid,EOpSrc::Invalid,TEXT("Exp2")),
+		FVectorVMOpInfo(EOp::log,EOpFlags::None,EOpSrc::Register,EOpSrc::Invalid,EOpSrc::Invalid,TEXT("Log")),
+		FVectorVMOpInfo(EOp::log2,EOpFlags::None,EOpSrc::Register,EOpSrc::Invalid,EOpSrc::Invalid,TEXT("Log base 2")),
+		FVectorVMOpInfo(EOp::sin,EOpFlags::Implemented,EOpSrc::Register,EOpSrc::Invalid,EOpSrc::Invalid,TEXT("Sin")),
+		FVectorVMOpInfo(EOp::sin, EOpFlags::None, EOpSrc::Const, EOpSrc::Invalid, EOpSrc::Invalid, TEXT("sini")),
+		FVectorVMOpInfo(EOp::cos, EOpFlags::None, EOpSrc::Register, EOpSrc::Invalid, EOpSrc::Invalid, TEXT("Cos")),
+		FVectorVMOpInfo(EOp::tan,EOpFlags::None,EOpSrc::Register,EOpSrc::Invalid,EOpSrc::Invalid,TEXT("Tan")),
+		FVectorVMOpInfo(EOp::asin,EOpFlags::None,EOpSrc::Register,EOpSrc::Invalid,EOpSrc::Invalid,TEXT("Arcsin")),
+		FVectorVMOpInfo(EOp::acos,EOpFlags::None,EOpSrc::Register,EOpSrc::Invalid,EOpSrc::Invalid,TEXT("Arccos")),
+		FVectorVMOpInfo(EOp::atan,EOpFlags::None,EOpSrc::Register,EOpSrc::Invalid,EOpSrc::Invalid,TEXT("Arctan")),
+		FVectorVMOpInfo(EOp::atan2,EOpFlags::None,EOpSrc::Register,EOpSrc::Register,EOpSrc::Invalid,TEXT("Arctan2")),
+		FVectorVMOpInfo(EOp::ceil,EOpFlags::None,EOpSrc::Register,EOpSrc::Invalid,EOpSrc::Invalid,TEXT("Round up")),
+		FVectorVMOpInfo(EOp::floor,EOpFlags::None,EOpSrc::Register,EOpSrc::Invalid,EOpSrc::Invalid,TEXT("Round down")),
+		FVectorVMOpInfo(EOp::fmod,EOpFlags::None,EOpSrc::Register,EOpSrc::Register,EOpSrc::Invalid,TEXT("Modulo")),
+		FVectorVMOpInfo(EOp::frac,EOpFlags::None,EOpSrc::Register,EOpSrc::Invalid,EOpSrc::Invalid,TEXT("Fractional")),
+		FVectorVMOpInfo(EOp::trunc,EOpFlags::None,EOpSrc::Register,EOpSrc::Invalid,EOpSrc::Invalid,TEXT("Trunc")),
+
+		FVectorVMOpInfo(EOp::clamp,EOpFlags::Implemented,EOpSrc::Register,EOpSrc::Register,EOpSrc::Register,TEXT("Clamp")),
 		FVectorVMOpInfo(EOp::clamp,EOpFlags::None,EOpSrc::Register,EOpSrc::Const,EOpSrc::Register,TEXT("clampir")),
 		FVectorVMOpInfo(EOp::clamp,EOpFlags::None,EOpSrc::Register,EOpSrc::Register,EOpSrc::Const,TEXT("clampri")),
 		FVectorVMOpInfo(EOp::clamp,EOpFlags::None,EOpSrc::Register,EOpSrc::Const,EOpSrc::Const,TEXT("clampii")),
 
-		FVectorVMOpInfo(EOp::min,EOpFlags::Implemented,EOpSrc::Register,EOpSrc::Register,EOpSrc::Invalid,TEXT("min")),
+		FVectorVMOpInfo(EOp::min,EOpFlags::Implemented,EOpSrc::Register,EOpSrc::Register,EOpSrc::Invalid,TEXT("Min")),
 		FVectorVMOpInfo(EOp::min,EOpFlags::None,EOpSrc::Register,EOpSrc::Const,EOpSrc::Invalid,TEXT("mini")),
 
-		FVectorVMOpInfo(EOp::max,EOpFlags::Implemented,EOpSrc::Register,EOpSrc::Register,EOpSrc::Invalid,TEXT("max")),
+		FVectorVMOpInfo(EOp::max,EOpFlags::Implemented,EOpSrc::Register,EOpSrc::Register,EOpSrc::Invalid,TEXT("Max")),
 		FVectorVMOpInfo(EOp::max,EOpFlags::None,EOpSrc::Register,EOpSrc::Const,EOpSrc::Invalid,TEXT("maxi")),
 
-		FVectorVMOpInfo(EOp::pow,EOpFlags::Implemented,EOpSrc::Register,EOpSrc::Register,EOpSrc::Invalid,TEXT("pow")),
+		FVectorVMOpInfo(EOp::pow,EOpFlags::Implemented,EOpSrc::Register,EOpSrc::Register,EOpSrc::Invalid,TEXT("Pow")),
 		FVectorVMOpInfo(EOp::pow,EOpFlags::None,EOpSrc::Register,EOpSrc::Const,EOpSrc::Invalid,TEXT("powi")),
 
-		FVectorVMOpInfo(EOp::add,EOpFlags::None,EOpSrc::Register,EOpSrc::Invalid,EOpSrc::Invalid,TEXT("sign")),
+		FVectorVMOpInfo(EOp::add,EOpFlags::None,EOpSrc::Register,EOpSrc::Invalid,EOpSrc::Invalid,TEXT("Sign")),
 
-		FVectorVMOpInfo(EOp::step,EOpFlags::None,EOpSrc::Register,EOpSrc::Register,EOpSrc::Invalid,TEXT("step")),
+		FVectorVMOpInfo(EOp::step,EOpFlags::None,EOpSrc::Register,EOpSrc::Register,EOpSrc::Invalid,TEXT("Step")),
 		FVectorVMOpInfo(EOp::step,EOpFlags::None,EOpSrc::Register,EOpSrc::Const,EOpSrc::Invalid,TEXT("stepi")),
 
 		FVectorVMOpInfo(EOp::add,EOpFlags::None,EOpSrc::Invalid,EOpSrc::Invalid,EOpSrc::Invalid,TEXT("tex1d")),
-		FVectorVMOpInfo(EOp::add,EOpFlags::None,EOpSrc::Invalid,EOpSrc::Invalid,EOpSrc::Invalid,TEXT("invalid"))
+
+		FVectorVMOpInfo(EOp::dot, EOpFlags::Implemented, EOpSrc::Register, EOpSrc::Register, EOpSrc::Invalid, TEXT("Dot Product")),
+		FVectorVMOpInfo(EOp::cross, EOpFlags::Implemented|EOpFlags::Commutative, EOpSrc::Register, EOpSrc::Register, EOpSrc::Invalid, TEXT("Cross Product")),
+		FVectorVMOpInfo(EOp::cross, EOpFlags::None, EOpSrc::Register, EOpSrc::Const, EOpSrc::Invalid, TEXT("Cross Product with const")),
+
+		FVectorVMOpInfo(EOp::normalize, EOpFlags::Implemented, EOpSrc::Register, EOpSrc::Invalid, EOpSrc::Invalid, TEXT("Normalize")),
+		FVectorVMOpInfo(EOp::random, EOpFlags::Implemented, EOpSrc::Const, EOpSrc::Invalid, EOpSrc::Invalid, TEXT("Random")),
+
+		FVectorVMOpInfo(EOp::length, EOpFlags::Implemented, EOpSrc::Register, EOpSrc::Invalid, EOpSrc::Invalid, TEXT("Vector Length")),
+		FVectorVMOpInfo(EOp::length, EOpFlags::None, EOpSrc::Const, EOpSrc::Invalid, EOpSrc::Invalid, TEXT("Vector Length (const)")),
+	
+		FVectorVMOpInfo(EOp::add, EOpFlags::None, EOpSrc::Invalid, EOpSrc::Invalid, EOpSrc::Invalid, TEXT("invalid"))
 	};
 } // namespace VectorVM
 
@@ -566,12 +606,12 @@ bool FVectorVMTest::RunTest(const FString& Parameters)
 {
 	uint8 TestCode[] =
 	{
-		VectorVM::EOp::mul,		0x00, 0x08, 0x08,       // mul r0, r8, r8
-		VectorVM::EOp::mad,		0x01, 0x09, 0x09, 0x00, // mad r1, r9, r9, r0
-		VectorVM::EOp::mad,		0x00, 0x0a, 0x0a, 0x01, // mad r0, r10, r10, r1
-		VectorVM::EOp::addi,	0x01, 0x00, 0x01,       // addi r1, r0, c1
-		VectorVM::EOp::neg,		0x00, 0x01,             // neg r0, r1
-		VectorVM::EOp::clampii, 0x28, 0x00, 0x02, 0x03, // clampii r40, r0, c2, c3
+		VectorVM::EOp::mul,		0x00, SRCOP_RRR, 0x08, 0x08,       // mul r0, r8, r8
+		VectorVM::EOp::mad,		0x01, SRCOP_RRR, 0x09, 0x09, 0x00, // mad r1, r9, r9, r0
+		VectorVM::EOp::mad,		0x00, SRCOP_RRR, 0x0a, 0x0a, 0x01, // mad r0, r10, r10, r1
+		VectorVM::EOp::add,		0x01, SRCOP_RRC, 0x00, 0x01,       // addi r1, r0, c1
+		VectorVM::EOp::neg,		0x00, SRCOP_RRR, 0x01,             // neg r0, r1
+		VectorVM::EOp::clamp,	0x28, SRCOP_CCR, 0x00, 0x02, 0x03, // clampii r40, r0, c2, c3
 		0x00 // terminator
 	};
 
@@ -590,11 +630,11 @@ bool FVectorVMTest::RunTest(const FString& Parameters)
 		reinterpret_cast<float*>(InputRegisters[2])[i] = static_cast<float>(i);
 	}
 
-	float ConstantTable[VectorVM::MaxConstants];
-	ConstantTable[0] = 0.0f;
-	ConstantTable[1] = 5.0f;
-	ConstantTable[2] = -20.0f;
-	ConstantTable[3] = 20.0f;
+	FVector4 ConstantTable[VectorVM::MaxConstants];
+	ConstantTable[0] = FVector4(0.0f, 0.0f, 0.0f, 0.0f);
+	ConstantTable[1] = FVector4(5.0f, 5.0f, 5.0f, 5.0f);
+	ConstantTable[2] = FVector4(-20.0f, -20.0f, -20.0f, -20.0f);
+	ConstantTable[3] = FVector4(20.0f, 20.0f, 20.0f, 20.0f);
 
 	VectorVM::Exec(
 		TestCode,
diff --git a/Engine/Source/Runtime/VectorVM/Public/VectorVM.h b/Engine/Source/Runtime/VectorVM/Public/VectorVM.h
--- a/Engine/Source/Runtime/VectorVM/Public/VectorVM.h
+++ b/Engine/Source/Runtime/VectorVM/Public/VectorVM.h
@@ -7,20 +7,24 @@
 #pragma once
 #include "Core.h"
 
+
+
 namespace VectorVM
 {
 	/** Constants. */
 	enum
 	{
-		NumTempRegisters = 8,
-		MaxInputRegisters = 32,
+		NumTempRegisters = 64,
+		MaxInputRegisters = 64,
 		MaxOutputRegisters = MaxInputRegisters,
+		MaxConstants = 256,
 		FirstInputRegister = NumTempRegisters,
 		FirstOutputRegister = FirstInputRegister + MaxInputRegisters,
-		MaxRegisters = NumTempRegisters + MaxInputRegisters + MaxOutputRegisters,
-		MaxConstants = 256,
+		FirstConstantRegister = FirstOutputRegister + MaxOutputRegisters,
+		MaxRegisters = NumTempRegisters + MaxInputRegisters + MaxOutputRegisters + MaxConstants
 	};
 
+
 	/** List of opcodes supported by the VM. */
 	namespace EOp
 	{
@@ -43,7 +47,7 @@ namespace VectorVM
 			lerpirr, 
 			lerprir, 
 			lerprri,
-			lerprii,
+			lerpiir,
 			rcp,
 			rsq,
 			sqrt,
@@ -54,6 +58,7 @@ namespace VectorVM
 			log,
 			log2,
 			sin,
+			sini,
 			cos,
 			tan,
 			asin,
@@ -79,6 +84,13 @@ namespace VectorVM
 			step,
 			stepi,
 			tex1d,
+			dot,
+			cross,
+			crossi,
+			normalize,
+			random,
+			length,
+			lengthi,
 			NumOpcodes
 		};
 	} // namespace EOp
@@ -152,7 +164,7 @@ namespace VectorVM
 		int32 NumInputRegisters,
 		VectorRegister** OutputRegisters,
 		int32 NumOutputRegisters,
-		float const* ConstantTable,
+		FVector4 const* ConstantTable,
 		int32 NumVectors
 		);
 
