Adding sin4 opcode to VectorVM. Fixing unit test which was causing all sorts of havoc with cooking

diff --git a/Engine/Source/Runtime/VectorVM/Private/VectorVM.cpp b/Engine/Source/Runtime/VectorVM/Private/VectorVM.cpp
--- a/Engine/Source/Runtime/VectorVM/Private/VectorVM.cpp
+++ b/Engine/Source/Runtime/VectorVM/Private/VectorVM.cpp
@@ -338,6 +338,20 @@ struct FVectorKernelSin : public TUnaryVectorKernel<FVectorKernelSin>
 	}
 };
 
+struct FVectorKernelSin4 : public TUnaryVectorKernel<FVectorKernelSin>
+{
+	static void FORCEINLINE DoKernel(VectorRegister* RESTRICT Dst, VectorRegister Src0)
+	{
+		float const* FloatSrc0 = reinterpret_cast<float const*>(&Src0);
+		float sn1 = FMath::Sin(*FloatSrc0++ * 3.14f);		// [0;1] takes us through half a period
+		float sn2 = FMath::Sin(*FloatSrc0++ * 3.14f);		// [0;1] takes us through half a period
+		float sn3 = FMath::Sin(*FloatSrc0++ * 3.14f);		// [0;1] takes us through half a period
+		float sn4 = FMath::Sin(*FloatSrc0++ * 3.14f);		// [0;1] takes us through half a period
+		*Dst = MakeVectorRegister(sn1, sn2, sn3, sn4);
+	}
+};
+
+
 struct FVectorKernelDot : public TBinaryVectorKernel<FVectorKernelDot>
 {
 	static void FORCEINLINE DoKernel(VectorRegister* RESTRICT Dst, VectorRegister Src0, VectorRegister Src1)
@@ -478,6 +492,7 @@ void VectorVM::Exec(
 			case EOp::max: FVectorKernelMax::Exec(Context); break;
 			case EOp::pow: FVectorKernelPow::Exec(Context); break;
 			case EOp::sin: FVectorKernelSin::Exec(Context); break;
+			case EOp::sin4: FVectorKernelSin4::Exec(Context); break;
 			case EOp::dot: FVectorKernelDot::Exec(Context); break;
 			case EOp::length: FVectorKernelLength::Exec(Context); break;
 			case EOp::cross: FVectorKernelCross::Exec(Context); break;
@@ -580,7 +595,10 @@ namespace VectorVM
 
 		FVectorVMOpInfo(EOp::length, EOpFlags::Implemented, EOpSrc::Register, EOpSrc::Invalid, EOpSrc::Invalid, TEXT("Vector Length")),
 		FVectorVMOpInfo(EOp::length, EOpFlags::None, EOpSrc::Const, EOpSrc::Invalid, EOpSrc::Invalid, TEXT("Vector Length (const)")),
-	
+
+		FVectorVMOpInfo(EOp::sin4, EOpFlags::Implemented, EOpSrc::Register, EOpSrc::Invalid, EOpSrc::Invalid, TEXT("Sin4")),
+		FVectorVMOpInfo(EOp::sin4, EOpFlags::None, EOpSrc::Const, EOpSrc::Invalid, EOpSrc::Invalid, TEXT("sin4i")),
+
 		FVectorVMOpInfo(EOp::add, EOpFlags::None, EOpSrc::Invalid, EOpSrc::Invalid, EOpSrc::Invalid, TEXT("invalid"))
 	};
 } // namespace VectorVM
@@ -606,12 +624,12 @@ bool FVectorVMTest::RunTest(const FString& Parameters)
 {
 	uint8 TestCode[] =
 	{
-		VectorVM::EOp::mul,		0x00, SRCOP_RRR, 0x08, 0x08,       // mul r0, r8, r8
-		VectorVM::EOp::mad,		0x01, SRCOP_RRR, 0x09, 0x09, 0x00, // mad r1, r9, r9, r0
-		VectorVM::EOp::mad,		0x00, SRCOP_RRR, 0x0a, 0x0a, 0x01, // mad r0, r10, r10, r1
-		VectorVM::EOp::add,		0x01, SRCOP_RRC, 0x00, 0x01,       // addi r1, r0, c1
+		VectorVM::EOp::mul, 0x00, SRCOP_RRR, 0x0 + VectorVM::NumTempRegisters, 0x0 + VectorVM::NumTempRegisters,       // mul r0, r8, r8
+		VectorVM::EOp::mad, 0x01, SRCOP_RRR, 0x01 + VectorVM::NumTempRegisters, 0x01 + VectorVM::NumTempRegisters, 0x00, // mad r1, r9, r9, r0
+		VectorVM::EOp::mad, 0x00, SRCOP_RRR, 0x02 + VectorVM::NumTempRegisters, 0x02 + VectorVM::NumTempRegisters, 0x01, // mad r0, r10, r10, r1
+		VectorVM::EOp::add,		0x01, SRCOP_RCR, 0x00, 0x01,       // addi r1, r0, c1
 		VectorVM::EOp::neg,		0x00, SRCOP_RRR, 0x01,             // neg r0, r1
-		VectorVM::EOp::clamp,	0x28, SRCOP_CCR, 0x00, 0x02, 0x03, // clampii r40, r0, c2, c3
+		VectorVM::EOp::clamp,	VectorVM::FirstOutputRegister, SRCOP_CCR, 0x00, 0x02, 0x03, // clampii r40, r0, c2, c3
 		0x00 // terminator
 	};
 
diff --git a/Engine/Source/Runtime/VectorVM/Public/VectorVM.h b/Engine/Source/Runtime/VectorVM/Public/VectorVM.h
--- a/Engine/Source/Runtime/VectorVM/Public/VectorVM.h
+++ b/Engine/Source/Runtime/VectorVM/Public/VectorVM.h
@@ -91,6 +91,8 @@ namespace VectorVM
 			random,
 			length,
 			lengthi,
+			sin4, 
+			sin4i,
 			NumOpcodes
 		};
 	} // namespace EOp
