Fix narrowing conversion with explicit cast

#codereveiw Olaf.Piesche

diff --git a/Engine/Source/Runtime/VectorVM/Private/VectorVM.cpp b/Engine/Source/Runtime/VectorVM/Private/VectorVM.cpp
--- a/Engine/Source/Runtime/VectorVM/Private/VectorVM.cpp
+++ b/Engine/Source/Runtime/VectorVM/Private/VectorVM.cpp
@@ -464,7 +464,7 @@ struct FVectorKernelNoise : public TUnaryVectorKernel<FVectorKernelNoise>
 				}
 			}
 
-			const uint32 IntCoords[3] = { FX, FY, FZ };
+			const uint32 IntCoords[3] = { static_cast<uint32>(FX), static_cast<uint32>(FY), static_cast<uint32>(FZ) };
 			const float Fractionals[3] = { FX - IntCoords[0], FY - IntCoords[1], FZ - IntCoords[2] };
 			VectorRegister Alpha = MakeVectorRegister(Fractionals[0], Fractionals[0], Fractionals[0], Fractionals[0]);
 
