More compile fixes. Will be at least 1 more as there's an xbone error I've not sussed yet.

diff --git a/Engine/Source/Editor/UnrealEd/Private/Niagara/NiagaraEditorCommon.cpp b/Engine/Source/Editor/UnrealEd/Private/Niagara/NiagaraEditorCommon.cpp
--- a/Engine/Source/Editor/UnrealEd/Private/Niagara/NiagaraEditorCommon.cpp
+++ b/Engine/Source/Editor/UnrealEd/Private/Niagara/NiagaraEditorCommon.cpp
@@ -557,18 +557,18 @@ TMap<FName, FNiagaraOpInfo>& FNiagaraOpInfo::GetOpInfoMap()
 
 //Debugging variant that checks inputs and outputs match expectations of FNiagaraOpInfo.
 #if DO_CHECK
-#define NiagaraOp(OpName) void INiagaraCompiler::##OpName(INiagaraCompiler* Compiler, TArray<TNiagaraExprPtr>& InputExpressions, TArray<TNiagaraExprPtr>& OutputExpressions)\
+#define NiagaraOp(OpName) void INiagaraCompiler::OpName(INiagaraCompiler* Compiler, TArray<TNiagaraExprPtr>& InputExpressions, TArray<TNiagaraExprPtr>& OutputExpressions)\
 {\
-	Compiler->CheckInputs(FNiagaraOpInfo::##OpName##, InputExpressions); \
+	Compiler->CheckInputs(FNiagaraOpInfo::OpName, InputExpressions); \
 	return Compiler->OpName##_Internal(InputExpressions, OutputExpressions); \
-	Compiler->CheckOutputs(FNiagaraOpInfo::##OpName##, OutputExpressions); \
+	Compiler->CheckOutputs(FNiagaraOpInfo::OpName, OutputExpressions); \
 }
 NiagaraOpList
 #undef NiagaraOp
 
 #else//DO_CHECK
 
-#define NiagaraOp(OpName) void INiagaraCompiler::##OpName(INiagaraCompiler* Compiler, TArray<TNiagaraExprPtr>& InputExpressions, TArray<TNiagaraExprPtr>& OutputExpressions)\
+#define NiagaraOp(OpName) void INiagaraCompiler::OpName(INiagaraCompiler* Compiler, TArray<TNiagaraExprPtr>& InputExpressions, TArray<TNiagaraExprPtr>& OutputExpressions)\
 { \
 	return Compiler->OpName##_Internal(InputExpressions, OutputExpressions); \
 }
diff --git a/Engine/Source/Runtime/Engine/Classes/Niagara/NiagaraConstantSet.h b/Engine/Source/Runtime/Engine/Classes/Niagara/NiagaraConstantSet.h
--- a/Engine/Source/Runtime/Engine/Classes/Niagara/NiagaraConstantSet.h
+++ b/Engine/Source/Runtime/Engine/Classes/Niagara/NiagaraConstantSet.h
@@ -23,7 +23,7 @@ public:
 		ScalarConstants.FindOrAdd(Name) = Sc;
 	}
 
-	void SetOrAdd(FName Name, FVector4 Vc)
+	void SetOrAdd(FName Name, const FVector4& Vc)
 	{
 		VectorConstants.FindOrAdd(Name) = Vc;
 	}
@@ -144,7 +144,7 @@ public:
 		}
 	}
 
-	void SetOrAdd(FName Name, FVector4 Vc)
+	void SetOrAdd(FName Name, const FVector4& Vc)
 	{
 		int32 Idx = VectorNames.Find(Name);
 		if (Idx == INDEX_NONE)
diff --git a/Engine/Source/Runtime/Engine/Classes/Niagara/NiagaraScriptConstantData.h b/Engine/Source/Runtime/Engine/Classes/Niagara/NiagaraScriptConstantData.h
--- a/Engine/Source/Runtime/Engine/Classes/Niagara/NiagaraScriptConstantData.h
+++ b/Engine/Source/Runtime/Engine/Classes/Niagara/NiagaraScriptConstantData.h
@@ -35,7 +35,7 @@ struct FNiagaraScriptConstantData
 		InternalConstants.SetOrAdd(Name, Sc);
 	}
 
-	void SetOrAddInternal(FName Name, FVector4 Vc)
+	void SetOrAddInternal(FName Name, const FVector4& Vc)
 	{
 		InternalConstants.SetOrAdd(Name, Vc);
 	}
@@ -50,7 +50,7 @@ struct FNiagaraScriptConstantData
 		ExternalConstants.SetOrAdd(Name, Sc);
 	}
 
-	void SetOrAddExternal(FName Name, FVector4 Vc)
+	void SetOrAddExternal(FName Name, const FVector4& Vc)
 	{
 		ExternalConstants.SetOrAdd(Name, Vc);
 	}
diff --git a/Engine/Source/Runtime/Engine/Private/NiagaraComponent.cpp b/Engine/Source/Runtime/Engine/Private/NiagaraComponent.cpp
--- a/Engine/Source/Runtime/Engine/Private/NiagaraComponent.cpp
+++ b/Engine/Source/Runtime/Engine/Private/NiagaraComponent.cpp
@@ -195,7 +195,7 @@ public:
 		Constants.SetOrAdd(ConstantName,Value);
 	}
 
-	void SetConstant(FName ConstantName, const FVector4 Value)
+	void SetConstant(FName ConstantName, const FVector4& Value)
 	{
 		Constants.SetOrAdd(ConstantName, Value);
 	}
diff --git a/Engine/Source/Runtime/VectorVM/Private/VectorVM.cpp b/Engine/Source/Runtime/VectorVM/Private/VectorVM.cpp
--- a/Engine/Source/Runtime/VectorVM/Private/VectorVM.cpp
+++ b/Engine/Source/Runtime/VectorVM/Private/VectorVM.cpp
@@ -90,17 +90,18 @@ struct FConstantHandler
 	FConstantHandler(FVectorVMContext& Context)
 		: Constant(DecodeConstant(Context))
 	{}
-	FORCEINLINE VectorRegister Get(){ return Constant; }
+	FORCEINLINE const VectorRegister& Get(){ return Constant; }
 };
 
 /** Handles reading of a register, advancing the pointer with each read. */
 struct FRegisterHandler
 {
 	VectorRegister* Register;
+	VectorRegister RegisterValue;
 	FRegisterHandler(FVectorVMContext& Context)
 		: Register(DecodeRegister(Context))
 	{}
-	FORCEINLINE VectorRegister Get(){ return *Register++; }
+	FORCEINLINE const VectorRegister& Get(){ RegisterValue = *Register++;  return RegisterValue; }
 };
 
 template<typename Kernel, typename Arg0Handler>
@@ -520,7 +521,10 @@ struct FVectorKernelSplat : public TUnaryVectorKernel<FVectorKernelSplat<Compone
 template<int32 Cmp0, int32 Cmp1, int32 Cmp2, int32 Cmp3>
 struct FVectorKernelCompose : public TQuinaryVectorKernel<FVectorKernelCompose<Cmp0, Cmp1, Cmp2, Cmp3>>
 {
-	static void FORCEINLINE DoKernel(VectorRegister* RESTRICT Dst, VectorRegister Src0, VectorRegister Src1, VectorRegister Src2, VectorRegister Src3)
+	//Passing as const refs as some compilers cant handle > 3 aligned vectorregister params.
+	//inlined so shouldn't impact perf
+	//Todo: ^^^^ test this
+	static void FORCEINLINE DoKernel(VectorRegister* RESTRICT Dst, const VectorRegister& Src0, const VectorRegister& Src1, const VectorRegister& Src2, const VectorRegister& Src3)
 	{
 		//TODO - There's probably a faster way to do this.
 		VectorRegister Tmp0 = VectorShuffle(Src0, Src1, Cmp0, Cmp0, Cmp1, Cmp1);
