Hopefully fixing compiler error on XBone.

diff --git a/Engine/Source/Runtime/VectorVM/Private/VectorVM.cpp b/Engine/Source/Runtime/VectorVM/Private/VectorVM.cpp
--- a/Engine/Source/Runtime/VectorVM/Private/VectorVM.cpp
+++ b/Engine/Source/Runtime/VectorVM/Private/VectorVM.cpp
@@ -535,7 +535,7 @@ struct FVectorKernelCompose : public TQuinaryVectorKernel<FVectorKernelCompose<C
 
 struct FVectorKernelOutput : public TUnaryVectorKernel<FVectorKernelOutput>
 {
-	static void FORCEINLINE DoKernel(VectorRegister* RESTRICT Dst, VectorRegister Src0)
+	static void FORCEINLINE DoKernel(VectorRegister* Dst, VectorRegister Src0)
 	{
 		VectorStoreAlignedStreamed(Src0, Dst);
 	}
