Stubbed out support for all VM ops. Most need a vectorized implementation.
Sin and Cos implementation is done.

diff --git a/Engine/Source/Editor/NiagaraEditor/Private/NiagaraCompiler_VectorVM.cpp b/Engine/Source/Editor/NiagaraEditor/Private/NiagaraCompiler_VectorVM.cpp
--- a/Engine/Source/Editor/NiagaraEditor/Private/NiagaraCompiler_VectorVM.cpp
+++ b/Engine/Source/Editor/NiagaraEditor/Private/NiagaraCompiler_VectorVM.cpp
@@ -457,11 +457,6 @@ void FNiagaraCompiler_VectorVM::Length_Internal(TArray<TNiagaraExprPtr>& InputEx
 	OutputExpressions.Add(Expression_VMNative(VectorVM::EOp::length, InputExpressions));
 }
 
-void FNiagaraCompiler_VectorVM::Sin4_Internal(TArray<TNiagaraExprPtr>& InputExpressions, TArray<TNiagaraExprPtr>& OutputExpressions)
-{
-	OutputExpressions.Add(Expression_VMNative(VectorVM::EOp::sin4, InputExpressions));
-}
-
 void FNiagaraCompiler_VectorVM::Noise_Internal(TArray<TNiagaraExprPtr>& InputExpressions, TArray<TNiagaraExprPtr>& OutputExpressions)
 {
 	OutputExpressions.Add(Expression_VMNative(VectorVM::EOp::noise, InputExpressions));
diff --git a/Engine/Source/Editor/UnrealEd/Private/Niagara/NiagaraEditorCommon.cpp b/Engine/Source/Editor/UnrealEd/Private/Niagara/NiagaraEditorCommon.cpp
--- a/Engine/Source/Editor/UnrealEd/Private/Niagara/NiagaraEditorCommon.cpp
+++ b/Engine/Source/Editor/UnrealEd/Private/Niagara/NiagaraEditorCommon.cpp
@@ -21,6 +21,9 @@ void FNiagaraOpInfo::Init()
 	FName Y(TEXT("Y"));	FText YText = NSLOCTEXT("NiagaraOpInfo", "Second Vector Component", "Y");
 	FName Z(TEXT("Z"));	FText ZText = NSLOCTEXT("NiagaraOpInfo", "Third Vector Component", "Z");
 	FName W(TEXT("W"));	FText WText = NSLOCTEXT("NiagaraOpInfo", "Fourth Vector Component", "W");
+	
+	FText MinText = NSLOCTEXT("NiagaraOpInfo", "Min", "Min");
+	FText MaxText = NSLOCTEXT("NiagaraOpInfo", "Max", "Max");
 
 	TMap<FName, FNiagaraOpInfo>& OpInfos = GetOpInfoMap();
 
@@ -252,9 +255,9 @@ void FNiagaraOpInfo::Init()
 	Op->Name = Clamp;
 	Op->FriendlyName = NSLOCTEXT("NiagaraOpInfo", "Clamp Name", "Clamp");
 	Op->Description = NSLOCTEXT("NiagaraOpInfo", "Clamp Desc", "Clamp");
-	Op->Inputs.Add(FNiagaraOpInOutInfo(A, ENiagaraDataType::Vector, AText, AText, DefaultStr_VecZero));
-	Op->Inputs.Add(FNiagaraOpInOutInfo(B, ENiagaraDataType::Vector, BText, BText, DefaultStr_VecOne));
-	Op->Inputs.Add(FNiagaraOpInOutInfo(C, ENiagaraDataType::Vector, CText, CText, DefaultStr_VecOne));
+	Op->Inputs.Add(FNiagaraOpInOutInfo(X, ENiagaraDataType::Vector, XText, XText, DefaultStr_VecZero));
+	Op->Inputs.Add(FNiagaraOpInOutInfo(Min, ENiagaraDataType::Vector, MinText, MinText, DefaultStr_VecOne));
+	Op->Inputs.Add(FNiagaraOpInOutInfo(Max, ENiagaraDataType::Vector, MaxText, MaxText, DefaultStr_VecOne));
 	Op->Outputs.Add(FNiagaraOpInOutInfo(Result, ENiagaraDataType::Vector, ResultText, ResultText, DefaultStr_VecOne));
 	Op->OpDelegate.BindStatic(&INiagaraCompiler::Clamp);
 
@@ -262,8 +265,8 @@ void FNiagaraOpInfo::Init()
 	Op->Name = Min;
 	Op->FriendlyName = NSLOCTEXT("NiagaraOpInfo", "Min Name", "Min");
 	Op->Description = NSLOCTEXT("NiagaraOpInfo", "Min Desc", "Min");
-	Op->Inputs.Add(FNiagaraOpInOutInfo(A, ENiagaraDataType::Vector, AText, AText, DefaultStr_VecZero));
-	Op->Inputs.Add(FNiagaraOpInOutInfo(B, ENiagaraDataType::Vector, BText, BText, DefaultStr_VecOne));
+	Op->Inputs.Add(FNiagaraOpInOutInfo(X, ENiagaraDataType::Vector, XText, XText, DefaultStr_VecZero));
+	Op->Inputs.Add(FNiagaraOpInOutInfo(Min, ENiagaraDataType::Vector, MinText, MinText, DefaultStr_VecOne));
 	Op->Outputs.Add(FNiagaraOpInOutInfo(Result, ENiagaraDataType::Vector, ResultText, ResultText, DefaultStr_VecOne));
 	Op->OpDelegate.BindStatic(&INiagaraCompiler::Min);
 
@@ -271,8 +274,8 @@ void FNiagaraOpInfo::Init()
 	Op->Name = Max;
 	Op->FriendlyName = NSLOCTEXT("NiagaraOpInfo", "Max Name", "Max");
 	Op->Description = NSLOCTEXT("NiagaraOpInfo", "Max Desc", "Max");
-	Op->Inputs.Add(FNiagaraOpInOutInfo(A, ENiagaraDataType::Vector, AText, AText, DefaultStr_VecZero));
-	Op->Inputs.Add(FNiagaraOpInOutInfo(B, ENiagaraDataType::Vector, BText, BText, DefaultStr_VecOne));
+	Op->Inputs.Add(FNiagaraOpInOutInfo(X, ENiagaraDataType::Vector, XText, XText, DefaultStr_VecZero));
+	Op->Inputs.Add(FNiagaraOpInOutInfo(Max, ENiagaraDataType::Vector, MaxText, MaxText, DefaultStr_VecOne));
 	Op->Outputs.Add(FNiagaraOpInOutInfo(Result, ENiagaraDataType::Vector, ResultText, ResultText, DefaultStr_VecOne));
 	Op->OpDelegate.BindStatic(&INiagaraCompiler::Max);
 
@@ -298,7 +301,6 @@ void FNiagaraOpInfo::Init()
 	Op->FriendlyName = NSLOCTEXT("NiagaraOpInfo", "Step Name", "Step");
 	Op->Description = NSLOCTEXT("NiagaraOpInfo", "Step Desc", "Step");
 	Op->Inputs.Add(FNiagaraOpInOutInfo(A, ENiagaraDataType::Vector, AText, AText, DefaultStr_VecOne));
-	Op->Inputs.Add(FNiagaraOpInOutInfo(B, ENiagaraDataType::Vector, BText, BText, DefaultStr_VecOne));
 	Op->Outputs.Add(FNiagaraOpInOutInfo(Result, ENiagaraDataType::Vector, ResultText, ResultText, DefaultStr_VecOne));
 	Op->OpDelegate.BindStatic(&INiagaraCompiler::Step);
 
@@ -344,14 +346,6 @@ void FNiagaraOpInfo::Init()
 	Op->Outputs.Add(FNiagaraOpInOutInfo(Result, ENiagaraDataType::Vector, ResultText, ResultText, DefaultStr_VecOne));
 	Op->OpDelegate.BindStatic(&INiagaraCompiler::Length);
 
-	Op = &OpInfos.Add(Sin4);
-	Op->Name = Sin4;
-	Op->FriendlyName = NSLOCTEXT("NiagaraOpInfo", "Sin4 Name", "Sin4");
-	Op->Description = NSLOCTEXT("NiagaraOpInfo", "Sin4 Desc", "Sin4");
-	Op->Inputs.Add(FNiagaraOpInOutInfo(A, ENiagaraDataType::Vector, AText, AText, DefaultStr_VecZero));
-	Op->Outputs.Add(FNiagaraOpInOutInfo(Result, ENiagaraDataType::Vector, ResultText, ResultText, DefaultStr_VecZero));
-	Op->OpDelegate.BindStatic(&INiagaraCompiler::Sin4);
-
 	Op = &OpInfos.Add(Noise);
 	Op->Name = Noise;
 	Op->FriendlyName = NSLOCTEXT("NiagaraOpInfo", "Noise Name", "Noise");
diff --git a/Engine/Source/Editor/UnrealEd/Public/NiagaraEditorCommon.h b/Engine/Source/Editor/UnrealEd/Public/NiagaraEditorCommon.h
--- a/Engine/Source/Editor/UnrealEd/Public/NiagaraEditorCommon.h
+++ b/Engine/Source/Editor/UnrealEd/Public/NiagaraEditorCommon.h
@@ -44,7 +44,6 @@
 	NiagaraOp(Normalize)\
 	NiagaraOp(Random)\
 	NiagaraOp(Length)\
-	NiagaraOp(Sin4)\
 	NiagaraOp(Noise)\
 	NiagaraOp(Scalar)\
 	NiagaraOp(Vector)\
diff --git a/Engine/Source/Programs/UnrealLightmass/Private/LightmassCore/Math/LMMathSSE.h b/Engine/Source/Programs/UnrealLightmass/Private/LightmassCore/Math/LMMathSSE.h
--- a/Engine/Source/Programs/UnrealLightmass/Private/LightmassCore/Math/LMMathSSE.h
+++ b/Engine/Source/Programs/UnrealLightmass/Private/LightmassCore/Math/LMMathSSE.h
@@ -82,6 +82,7 @@ static const VectorRegister SSE_SIGN_MASK = MakeVectorRegister( (uint32)SIGN_BIT
 #undef VectorStoreFloat3
 #undef VectorStoreFloat1
 #undef VectorReplicate
+#undef VectorAbs
 #undef VectorNegate
 #undef VectorSwizzle
 #undef VectorLoadByte4Reverse
diff --git a/Engine/Source/Runtime/Core/Public/Math/UnrealMathDirectX.h b/Engine/Source/Runtime/Core/Public/Math/UnrealMathDirectX.h
--- a/Engine/Source/Runtime/Core/Public/Math/UnrealMathDirectX.h
+++ b/Engine/Source/Runtime/Core/Public/Math/UnrealMathDirectX.h
@@ -697,6 +697,129 @@ inline bool VectorContainsNaNOrInfinite(const VectorRegister& Vec)
 	return (XMVector4IsNaN( Vec ) || XMVector4IsInfinite( Vec ));
 }
 
+//TODO: Vectorize
+FORCEINLINE VectorRegister VectorExp(const VectorRegister& X)
+{
+	return MakeVectorRegister((float)exp(VectorGetComponent(X, 0)), (float)exp(VectorGetComponent(X, 1)), (float)exp(VectorGetComponent(X, 2)), (float)exp(VectorGetComponent(X, 3)));
+}
+
+//TODO: Vectorize
+FORCEINLINE VectorRegister VectorExp2(const VectorRegister& X)
+{
+	return DirectX::XMVectorExp2(X);
+}
+
+//TODO: Vectorize
+FORCEINLINE VectorRegister VectorLog(const VectorRegister& X)
+{
+	return MakeVectorRegister((float)log(VectorGetComponent(X, 0)), (float)log(VectorGetComponent(X, 1)), (float)log(VectorGetComponent(X, 2)), (float)log(VectorGetComponent(X, 3)));
+}
+
+//TODO: Vectorize
+FORCEINLINE VectorRegister VectorLog2(const VectorRegister& X)
+{
+	return DirectX::XMVectorLog2(X);
+}
+
+//TODO: Vectorize
+FORCEINLINE VectorRegister VectorSin(const VectorRegister& X)
+{
+	return MakeVectorRegister((float)sin(VectorGetComponent(X, 0)), (float)sin(VectorGetComponent(X, 1)), (float)sin(VectorGetComponent(X, 2)), (float)sin(VectorGetComponent(X, 3)));
+}
+
+//TODO: Vectorize
+FORCEINLINE VectorRegister VectorCos(const VectorRegister& X)
+{
+	return MakeVectorRegister((float)cos(VectorGetComponent(X, 0)), (float)cos(VectorGetComponent(X, 1)), (float)cos(VectorGetComponent(X, 2)), (float)cos(VectorGetComponent(X, 3)));
+}
+
+//TODO: Vectorize
+FORCEINLINE VectorRegister VectorTan(const VectorRegister& X)
+{
+	return MakeVectorRegister((float)tan(VectorGetComponent(X, 0)), (float)tan(VectorGetComponent(X, 1)), (float)tan(VectorGetComponent(X, 2)), (float)tan(VectorGetComponent(X, 3)));
+}
+
+//TODO: Vectorize
+FORCEINLINE VectorRegister VectorASin(const VectorRegister& X)
+{
+	return MakeVectorRegister((float)asin(VectorGetComponent(X, 0)), (float)asin(VectorGetComponent(X, 1)), (float)asin(VectorGetComponent(X, 2)), (float)asin(VectorGetComponent(X, 3)));
+}
+
+//TODO: Vectorize
+FORCEINLINE VectorRegister VectorACos(const VectorRegister& X)
+{
+	return MakeVectorRegister((float)acos(VectorGetComponent(X, 0)), (float)acos(VectorGetComponent(X, 1)), (float)acos(VectorGetComponent(X, 2)), (float)acos(VectorGetComponent(X, 3)));
+}
+
+//TODO: Vectorize
+FORCEINLINE VectorRegister VectorATan(const VectorRegister& X)
+{
+	return MakeVectorRegister((float)atan(VectorGetComponent(X, 0)), (float)atan(VectorGetComponent(X, 1)), (float)atan(VectorGetComponent(X, 2)), (float)atan(VectorGetComponent(X, 3)));
+}
+
+//TODO: Vectorize
+FORCEINLINE VectorRegister VectorATan2(const VectorRegister& X, const VectorRegister& Y)
+{
+	return MakeVectorRegister((float)atan2(VectorGetComponent(X, 0), VectorGetComponent(Y, 0)),
+		(float)atan2(VectorGetComponent(X, 1), VectorGetComponent(Y, 1)),
+		(float)atan2(VectorGetComponent(X, 2), VectorGetComponent(Y, 2)),
+		(float)atan2(VectorGetComponent(X, 3), VectorGetComponent(Y, 3)));
+}
+
+//TODO: Vectorize
+FORCEINLINE VectorRegister VectorCeil(const VectorRegister& X)
+{
+	return MakeVectorRegister((float)ceil(VectorGetComponent(X, 0)), (float)ceil(VectorGetComponent(X, 1)), (float)ceil(VectorGetComponent(X, 2)), (float)ceil(VectorGetComponent(X, 3)));
+}
+
+//TODO: Vectorize
+FORCEINLINE VectorRegister VectorFloor(const VectorRegister& X)
+{
+	return MakeVectorRegister((float)floor(VectorGetComponent(X, 0)), (float)floor(VectorGetComponent(X, 1)), (float)floor(VectorGetComponent(X, 2)), (float)floor(VectorGetComponent(X, 3)));
+}
+
+//TODO: Vectorize
+FORCEINLINE VectorRegister VectorTruncate(const VectorRegister& X)
+{
+	return MakeVectorRegister((float)(int32)VectorGetComponent(X, 0), (float)(int32)VectorGetComponent(X, 1), (float)(int32)VectorGetComponent(X, 2), (float)(int32)VectorGetComponent(X, 3));
+}
+
+//TODO: Vectorize
+FORCEINLINE VectorRegister VectorFractional(const VectorRegister& X)
+{
+	return VectorSubtract(X, VectorTruncate(X));
+}
+
+//TODO: Vectorize
+FORCEINLINE VectorRegister VectorMod(const VectorRegister& X, const VectorRegister& Y)
+{
+	return MakeVectorRegister((float)fmod(VectorGetComponent(X, 0), VectorGetComponent(Y, 0)),
+		(float)fmod(VectorGetComponent(X, 1), VectorGetComponent(Y, 1)),
+		(float)fmod(VectorGetComponent(X, 2), VectorGetComponent(Y, 2)),
+		(float)fmod(VectorGetComponent(X, 3), VectorGetComponent(Y, 3)));
+}
+
+//TODO: Vectorize
+FORCEINLINE VectorRegister VectorSign(const VectorRegister& X)
+{
+	return MakeVectorRegister(
+		(float)(VectorGetComponent(X, 0) >= 0.0f ? 1.0f : 0.0f),
+		(float)(VectorGetComponent(X, 1) >= 0.0f ? 1.0f : 0.0f),
+		(float)(VectorGetComponent(X, 2) >= 0.0f ? 1.0f : 0.0f),
+		(float)(VectorGetComponent(X, 3) >= 0.0f ? 1.0f : 0.0f));
+}
+
+//TODO: Vectorize
+FORCEINLINE VectorRegister VectorStep(const VectorRegister& X)
+{
+	return MakeVectorRegister(
+		(float)(VectorGetComponent(X, 0) >= 0.0f ? 1.0f : -1.0f),
+		(float)(VectorGetComponent(X, 1) >= 0.0f ? 1.0f : -1.0f),
+		(float)(VectorGetComponent(X, 2) >= 0.0f ? 1.0f : -1.0f),
+		(float)(VectorGetComponent(X, 3) >= 0.0f ? 1.0f : -1.0f));
+}
+
+
 // To be continued...
 
 #endif
diff --git a/Engine/Source/Runtime/Core/Public/Math/UnrealMathFPU.h b/Engine/Source/Runtime/Core/Public/Math/UnrealMathFPU.h
--- a/Engine/Source/Runtime/Core/Public/Math/UnrealMathFPU.h
+++ b/Engine/Source/Runtime/Core/Public/Math/UnrealMathFPU.h
@@ -57,13 +57,6 @@ FORCEINLINE VectorRegister MakeVectorRegister( float X, float Y, float Z, float
 	return Vec;
 }
 
-/** Vector that represents (0,0,0,0) */
-static const VectorRegister VECTOR_ZERO = MakeVectorRegister( 0.0f, 0.0f, 0.0f, 0.0f );
-
-/** Vector that represents (1,1,1,1) */
-static const VectorRegister VECTOR_ONE  = MakeVectorRegister( 1.0f, 1.0f, 1.0f, 1.0f );
-
-
 /*=============================================================================
  *	Constants:
  *============================================================================*/
@@ -80,14 +73,14 @@ static const VectorRegister VECTOR_ONE  = MakeVectorRegister( 1.0f, 1.0f, 1.0f,
  *
  * @return		VectorRegister(0.0f, 0.0f, 0.0f, 0.0f)
  */
-#define VectorZero()					VECTOR_ZERO
+#define VectorZero()					(GlobalVectorConstants::FloatZero)
 
 /**
  * Returns a vector with all ones.
  *
  * @return		VectorRegister(1.0f, 1.0f, 1.0f, 1.0f)
  */
-#define VectorOne()						VECTOR_ONE
+#define VectorOne()						(GlobalVectorConstants::FloatOne)
 
 /**
  * Loads 4 FLOATs from unaligned memory.
@@ -944,7 +937,11 @@ FORCEINLINE uint32 VectorAnyGreaterThan(const VectorRegister& Vec1, const Vector
  * @param ComponentIndex	Which component to get, X=0, Y=1, Z=2, W=3
  * @return					The component as a float
  */
-#define VectorGetComponent( Vec, ComponentIndex )	(((float*) &(Vec))[ComponentIndex])
+FORCEINLINE float VectorGetComponent(VectorRegister Vec, uint32 ComponentIndex)
+{
+	return (((float*)&(Vec))[ComponentIndex]);
+}
+
 
 /**
  * Sets the control register.
@@ -965,6 +962,130 @@ inline bool VectorContainsNaNOrInfinite(const VectorRegister& Vec)
 		!FMath::IsFinite(Vec.V[0]) || !FMath::IsFinite(Vec.V[1]) || !FMath::IsFinite(Vec.V[2]) || !FMath::IsFinite(Vec.V[3]);
 }
 
+
+//TODO: Vectorize
+FORCEINLINE VectorRegister VectorExp(const VectorRegister& X)
+{
+	return MakeVectorRegister((float)exp(VectorGetComponent(X, 0)), (float)exp(VectorGetComponent(X, 1)), (float)exp(VectorGetComponent(X, 2)), (float)exp(VectorGetComponent(X, 3)));
+}
+
+//TODO: Vectorize
+FORCEINLINE VectorRegister VectorExp2(const VectorRegister& X)
+{
+	return MakeVectorRegister((float)exp2(VectorGetComponent(X, 0)), (float)exp2(VectorGetComponent(X, 1)), (float)exp2(VectorGetComponent(X, 2)), (float)exp2(VectorGetComponent(X, 3)));
+}
+
+//TODO: Vectorize
+FORCEINLINE VectorRegister VectorLog(const VectorRegister& X)
+{
+	return MakeVectorRegister((float)log(VectorGetComponent(X, 0)), (float)log(VectorGetComponent(X, 1)), (float)log(VectorGetComponent(X, 2)), (float)log(VectorGetComponent(X, 3)));
+}
+
+//TODO: Vectorize
+FORCEINLINE VectorRegister VectorLog2(const VectorRegister& X)
+{
+	return MakeVectorRegister((float)log2(VectorGetComponent(X, 0)), (float)log2(VectorGetComponent(X, 1)), (float)log2(VectorGetComponent(X, 2)), (float)log2(VectorGetComponent(X, 3)));
+}
+
+//TODO: Vectorize
+FORCEINLINE VectorRegister VectorSin(const VectorRegister& X)
+{
+	return MakeVectorRegister((float)sin(VectorGetComponent(X, 0)), (float)sin(VectorGetComponent(X, 1)), (float)sin(VectorGetComponent(X, 2)), (float)sin(VectorGetComponent(X, 3)));
+}
+
+//TODO: Vectorize
+FORCEINLINE VectorRegister VectorCos(const VectorRegister& X)
+{
+	return MakeVectorRegister((float)cos(VectorGetComponent(X, 0)), (float)cos(VectorGetComponent(X, 1)), (float)cos(VectorGetComponent(X, 2)), (float)cos(VectorGetComponent(X, 3)));
+}
+
+//TODO: Vectorize
+FORCEINLINE VectorRegister VectorTan(const VectorRegister& X)
+{
+	return MakeVectorRegister((float)tan(VectorGetComponent(X, 0)), (float)tan(VectorGetComponent(X, 1)), (float)tan(VectorGetComponent(X, 2)), (float)tan(VectorGetComponent(X, 3)));
+}
+
+//TODO: Vectorize
+FORCEINLINE VectorRegister VectorASin(const VectorRegister& X)
+{
+	return MakeVectorRegister((float)asin(VectorGetComponent(X, 0)), (float)asin(VectorGetComponent(X, 1)), (float)asin(VectorGetComponent(X, 2)), (float)asin(VectorGetComponent(X, 3)));
+}
+
+//TODO: Vectorize
+FORCEINLINE VectorRegister VectorACos(const VectorRegister& X)
+{
+	return MakeVectorRegister((float)acos(VectorGetComponent(X, 0)), (float)acos(VectorGetComponent(X, 1)), (float)acos(VectorGetComponent(X, 2)), (float)acos(VectorGetComponent(X, 3)));
+}
+
+//TODO: Vectorize
+FORCEINLINE VectorRegister VectorATan(const VectorRegister& X)
+{
+	return MakeVectorRegister((float)atan(VectorGetComponent(X, 0)), (float)atan(VectorGetComponent(X, 1)), (float)atan(VectorGetComponent(X, 2)), (float)atan(VectorGetComponent(X, 3)));
+}
+
+//TODO: Vectorize
+FORCEINLINE VectorRegister VectorATan2(const VectorRegister& X, const VectorRegister& Y)
+{
+	return MakeVectorRegister((float)atan2(VectorGetComponent(X, 0), VectorGetComponent(Y, 0)),
+		(float)atan2(VectorGetComponent(X, 1), VectorGetComponent(Y, 1)),
+		(float)atan2(VectorGetComponent(X, 2), VectorGetComponent(Y, 2)),
+		(float)atan2(VectorGetComponent(X, 3), VectorGetComponent(Y, 3)));
+}
+
+//TODO: Vectorize
+FORCEINLINE VectorRegister VectorCeil(const VectorRegister& X)
+{
+	return MakeVectorRegister((float)ceil(VectorGetComponent(X, 0)), (float)ceil(VectorGetComponent(X, 1)), (float)ceil(VectorGetComponent(X, 2)), (float)ceil(VectorGetComponent(X, 3)));
+}
+
+//TODO: Vectorize
+FORCEINLINE VectorRegister VectorFloor(const VectorRegister& X)
+{
+	return MakeVectorRegister((float)floor(VectorGetComponent(X, 0)), (float)floor(VectorGetComponent(X, 1)), (float)floor(VectorGetComponent(X, 2)), (float)floor(VectorGetComponent(X, 3)));
+}
+
+//TODO: Vectorize
+FORCEINLINE VectorRegister VectorTruncate(const VectorRegister& X)
+{
+	return MakeVectorRegister((float)(int32)VectorGetComponent(X, 0), (float)(int32)VectorGetComponent(X, 1), (float)(int32)VectorGetComponent(X, 2), (float)(int32)VectorGetComponent(X, 3));
+}
+
+//TODO: Vectorize
+FORCEINLINE VectorRegister VectorFractional(const VectorRegister& X)
+{
+	return VectorSubtract(X, VectorTruncate(X));
+}
+
+//TODO: Vectorize
+FORCEINLINE VectorRegister VectorMod(const VectorRegister& X, const VectorRegister& Y)
+{
+	return MakeVectorRegister((float)fmod(VectorGetComponent(X, 0), VectorGetComponent(Y, 0)),
+		(float)fmod(VectorGetComponent(X, 1), VectorGetComponent(Y, 1)),
+		(float)fmod(VectorGetComponent(X, 2), VectorGetComponent(Y, 2)),
+		(float)fmod(VectorGetComponent(X, 3), VectorGetComponent(Y, 3)));
+}
+
+//TODO: Vectorize
+FORCEINLINE VectorRegister VectorSign(const VectorRegister& X)
+{
+	return MakeVectorRegister(
+		(float)(VectorGetComponent(X, 0) >= 0.0f ? 1.0f : 0.0f),
+		(float)(VectorGetComponent(X, 1) >= 0.0f ? 1.0f : 0.0f),
+		(float)(VectorGetComponent(X, 2) >= 0.0f ? 1.0f : 0.0f),
+		(float)(VectorGetComponent(X, 3) >= 0.0f ? 1.0f : 0.0f));
+}
+
+//TODO: Vectorize
+FORCEINLINE VectorRegister VectorStep(const VectorRegister& X)
+{
+	return MakeVectorRegister(
+		(float)(VectorGetComponent(X, 0) >= 0.0f ? 1.0f : -1.0f),
+		(float)(VectorGetComponent(X, 1) >= 0.0f ? 1.0f : -1.0f),
+		(float)(VectorGetComponent(X, 2) >= 0.0f ? 1.0f : -1.0f),
+		(float)(VectorGetComponent(X, 3) >= 0.0f ? 1.0f : -1.0f));
+}
+
+
 // To be continued...
 
 
diff --git a/Engine/Source/Runtime/Core/Public/Math/UnrealMathNeon.h b/Engine/Source/Runtime/Core/Public/Math/UnrealMathNeon.h
--- a/Engine/Source/Runtime/Core/Public/Math/UnrealMathNeon.h
+++ b/Engine/Source/Runtime/Core/Public/Math/UnrealMathNeon.h
@@ -58,10 +58,6 @@ FORCEINLINE VectorRegister MakeVectorRegister( float X, float Y, float Z, float
 	return Tmp.V;
 }
 
-static const VectorRegister Vec1110 = MakeVectorRegister( 1.0f, 1.0f, 1.0f, 0.0f );
-static const VectorRegister Vec0001 = MakeVectorRegister( 0.0f, 0.0f, 0.0f, 1.0f );
-static const VectorRegister Vec255 = MakeVectorRegister( 255.f, 255.f, 255.f, 255.f );
-
 /*
 #define VectorPermute(Vec1, Vec2, Mask) my_perm(Vec1, Vec2, Mask)
 
@@ -879,7 +875,7 @@ FORCEINLINE VectorRegister VectorLoadByte4Reverse( const void* Ptr )
  */
 FORCEINLINE void VectorStoreByte4( VectorRegister Vec, void* Ptr )
 {
-	uint16x8_t u16x8 = (uint16x8_t)vcvtq_u32_f32( VectorMin( Vec, Vec255 ) );
+	uint16x8_t u16x8 = (uint16x8_t)vcvtq_u32_f32(VectorMin(Vec, GlobalVectorConstants::Float255));
 	uint8x8_t u8x8 = (uint8x8_t)vget_low_u16( vuzpq_u16( u16x8, u16x8 ).val[0] );
 	u8x8 = vuzp_u8( u8x8, u8x8 ).val[0];
 	uint32_t buf[2];
@@ -997,6 +993,129 @@ inline bool VectorContainsNaNOrInfinite(const VectorRegister& Vec)
 	return false;
 }
 
+//TODO: Vectorize
+FORCEINLINE VectorRegister VectorExp(const VectorRegister& X)
+{
+	return MakeVectorRegister((float)exp(VectorGetComponent(X, 0)), (float)exp(VectorGetComponent(X, 1)), (float)exp(VectorGetComponent(X, 2)), (float)exp(VectorGetComponent(X, 3)));
+}
+
+//TODO: Vectorize
+FORCEINLINE VectorRegister VectorExp2(const VectorRegister& X)
+{
+	return MakeVectorRegister((float)exp2(VectorGetComponent(X, 0)), (float)exp2(VectorGetComponent(X, 1)), (float)exp2(VectorGetComponent(X, 2)), (float)exp2(VectorGetComponent(X, 3)));
+}
+
+//TODO: Vectorize
+FORCEINLINE VectorRegister VectorLog(const VectorRegister& X)
+{
+	return MakeVectorRegister((float)log(VectorGetComponent(X, 0)), (float)log(VectorGetComponent(X, 1)), (float)log(VectorGetComponent(X, 2)), (float)log(VectorGetComponent(X, 3)));
+}
+
+//TODO: Vectorize
+FORCEINLINE VectorRegister VectorLog2(const VectorRegister& X)
+{
+	return MakeVectorRegister((float)log2(VectorGetComponent(X, 0)), (float)log2(VectorGetComponent(X, 1)), (float)log2(VectorGetComponent(X, 2)), (float)log2(VectorGetComponent(X, 3)));
+}
+
+//TODO: Vectorize
+FORCEINLINE VectorRegister VectorSin(const VectorRegister& X)
+{
+	return MakeVectorRegister((float)sin(VectorGetComponent(X, 0)), (float)sin(VectorGetComponent(X, 1)), (float)sin(VectorGetComponent(X, 2)), (float)sin(VectorGetComponent(X, 3)));
+}
+
+//TODO: Vectorize
+FORCEINLINE VectorRegister VectorCos(const VectorRegister& X)
+{
+	return MakeVectorRegister((float)cos(VectorGetComponent(X, 0)), (float)cos(VectorGetComponent(X, 1)), (float)cos(VectorGetComponent(X, 2)), (float)cos(VectorGetComponent(X, 3)));
+}
+
+//TODO: Vectorize
+FORCEINLINE VectorRegister VectorTan(const VectorRegister& X)
+{
+	return MakeVectorRegister((float)tan(VectorGetComponent(X, 0)), (float)tan(VectorGetComponent(X, 1)), (float)tan(VectorGetComponent(X, 2)), (float)tan(VectorGetComponent(X, 3)));
+}
+
+//TODO: Vectorize
+FORCEINLINE VectorRegister VectorASin(const VectorRegister& X)
+{
+	return MakeVectorRegister((float)asin(VectorGetComponent(X, 0)), (float)asin(VectorGetComponent(X, 1)), (float)asin(VectorGetComponent(X, 2)), (float)asin(VectorGetComponent(X, 3)));
+}
+
+//TODO: Vectorize
+FORCEINLINE VectorRegister VectorACos(const VectorRegister& X)
+{
+	return MakeVectorRegister((float)acos(VectorGetComponent(X, 0)), (float)acos(VectorGetComponent(X, 1)), (float)acos(VectorGetComponent(X, 2)), (float)acos(VectorGetComponent(X, 3)));
+}
+
+//TODO: Vectorize
+FORCEINLINE VectorRegister VectorATan(const VectorRegister& X)
+{
+	return MakeVectorRegister((float)atan(VectorGetComponent(X, 0)), (float)atan(VectorGetComponent(X, 1)), (float)atan(VectorGetComponent(X, 2)), (float)atan(VectorGetComponent(X, 3)));
+}
+
+//TODO: Vectorize
+FORCEINLINE VectorRegister VectorATan2(const VectorRegister& X, const VectorRegister& Y)
+{
+	return MakeVectorRegister((float)atan2(VectorGetComponent(X, 0), VectorGetComponent(Y, 0)),
+		(float)atan2(VectorGetComponent(X, 1), VectorGetComponent(Y, 1)),
+		(float)atan2(VectorGetComponent(X, 2), VectorGetComponent(Y, 2)),
+		(float)atan2(VectorGetComponent(X, 3), VectorGetComponent(Y, 3)));
+}
+
+//TODO: Vectorize
+FORCEINLINE VectorRegister VectorCeil(const VectorRegister& X)
+{
+	return MakeVectorRegister((float)ceil(VectorGetComponent(X, 0)), (float)ceil(VectorGetComponent(X, 1)), (float)ceil(VectorGetComponent(X, 2)), (float)ceil(VectorGetComponent(X, 3)));
+}
+
+//TODO: Vectorize
+FORCEINLINE VectorRegister VectorFloor(const VectorRegister& X)
+{
+	return MakeVectorRegister((float)floor(VectorGetComponent(X, 0)), (float)floor(VectorGetComponent(X, 1)), (float)floor(VectorGetComponent(X, 2)), (float)floor(VectorGetComponent(X, 3)));
+}
+
+//TODO: Vectorize
+FORCEINLINE VectorRegister VectorTruncate(const VectorRegister& X)
+{
+	return MakeVectorRegister((float)(int32)VectorGetComponent(X, 0), (float)(int32)VectorGetComponent(X, 1), (float)(int32)VectorGetComponent(X, 2), (float)(int32)VectorGetComponent(X, 3));
+}
+
+//TODO: Vectorize
+FORCEINLINE VectorRegister VectorFractional(const VectorRegister& X)
+{
+	return VectorSubtract(X, VectorTruncate(X));
+}
+
+//TODO: Vectorize
+FORCEINLINE VectorRegister VectorMod(const VectorRegister& X, const VectorRegister& Y)
+{
+	return MakeVectorRegister((float)fmod(VectorGetComponent(X, 0), VectorGetComponent(Y, 0)),
+		(float)fmod(VectorGetComponent(X, 1), VectorGetComponent(Y, 1)),
+		(float)fmod(VectorGetComponent(X, 2), VectorGetComponent(Y, 2)),
+		(float)fmod(VectorGetComponent(X, 3), VectorGetComponent(Y, 3)));
+}
+
+//TODO: Vectorize
+FORCEINLINE VectorRegister VectorSign(const VectorRegister& X)
+{
+	return MakeVectorRegister(
+		(float)(VectorGetComponent(X, 0) >= 0.0f ? 1.0f : 0.0f),
+		(float)(VectorGetComponent(X, 1) >= 0.0f ? 1.0f : 0.0f),
+		(float)(VectorGetComponent(X, 2) >= 0.0f ? 1.0f : 0.0f),
+		(float)(VectorGetComponent(X, 3) >= 0.0f ? 1.0f : 0.0f));
+}
+
+//TODO: Vectorize
+FORCEINLINE VectorRegister VectorStep(const VectorRegister& X)
+{
+	return MakeVectorRegister(
+		(float)(VectorGetComponent(X, 0) >= 0.0f ? 1.0f : -1.0f),
+		(float)(VectorGetComponent(X, 1) >= 0.0f ? 1.0f : -1.0f),
+		(float)(VectorGetComponent(X, 2) >= 0.0f ? 1.0f : -1.0f),
+		(float)(VectorGetComponent(X, 3) >= 0.0f ? 1.0f : -1.0f));
+}
+
+
 // To be continued...
 
 
diff --git a/Engine/Source/Runtime/Core/Public/Math/UnrealMathSSE.h b/Engine/Source/Runtime/Core/Public/Math/UnrealMathSSE.h
--- a/Engine/Source/Runtime/Core/Public/Math/UnrealMathSSE.h
+++ b/Engine/Source/Runtime/Core/Public/Math/UnrealMathSSE.h
@@ -66,19 +66,6 @@ FORCEINLINE VectorRegister MakeVectorRegister( float X, float Y, float Z, float
 	return _mm_setr_ps( X, Y, Z, W );
 }
 
-/** Vector that represents (1,1,1,1) */
-static const VectorRegister SSE_ONE = MakeVectorRegister( 1.0f, 1.0f, 1.0f, 1.0f );
-
-/** Bitmask to AND out the XYZ components in a vector */
-static const VectorRegister SSE_XYZ_MASK = MakeVectorRegister( (uint32)0xffffffff, (uint32)0xffffffff, (uint32)0xffffffff, (uint32)0x00000000 );
-
-/** Bitmask to AND out the sign bit of each components in a vector */
-#define SIGN_BIT (~(1 << 31))
-static const VectorRegister SSE_SIGN_MASK = MakeVectorRegister( (uint32)SIGN_BIT, (uint32)SIGN_BIT, (uint32)SIGN_BIT, (uint32)SIGN_BIT );
-
-/** Vector full of positive infinity */
-static const VectorRegister SSE_INFINITY = MakeVectorRegister( (uint32)0x7F800000, (uint32)0x7F800000, (uint32)0x7F800000, (uint32)0x7F800000 );
-
 /*=============================================================================
  *	Constants:
  *============================================================================*/
@@ -101,7 +88,7 @@ static const VectorRegister SSE_INFINITY = MakeVectorRegister( (uint32)0x7F80000
  *
  * @return		VectorRegister(1.0f, 1.0f, 1.0f, 1.0f)
  */
-#define VectorOne()						SSE_ONE
+#define VectorOne()						(GlobalVectorConstants::FloatOne)
 
 /**
  * Returns an component from a vector.
@@ -110,7 +97,10 @@ static const VectorRegister SSE_INFINITY = MakeVectorRegister( (uint32)0x7F80000
  * @param ComponentIndex	Which component to get, X=0, Y=1, Z=2, W=3
  * @return					The component as a float
  */
-#define VectorGetComponent( Vec, ComponentIndex )	(((float*) &(Vec))[ComponentIndex])
+FORCEINLINE float VectorGetComponent(VectorRegister Vec, uint32 ComponentIndex)
+{
+	return (((float*)&(Vec))[ComponentIndex]);
+}
 
 /**
  * Loads 4 FLOATs from unaligned memory.
@@ -255,7 +245,7 @@ FORCEINLINE void VectorStoreFloat3( const VectorRegister& Vec, void* Ptr )
  * @param Vec			Source vector
  * @return				VectorRegister( abs(Vec.x), abs(Vec.y), abs(Vec.z), abs(Vec.w) )
  */
-#define VectorAbs( Vec )				_mm_and_ps(Vec, SSE_SIGN_MASK)
+#define VectorAbs( Vec )				_mm_and_ps(Vec, GlobalVectorConstants::SignMask)
 
 /**
  * Returns the negated value (component-wise).
@@ -561,7 +551,7 @@ FORCEINLINE VectorRegister VectorNormalize( const VectorRegister& Vector )
 * @param Vector	VectorRegister
 * @return		VectorRegister(X, Y, Z, 0.0f)
 */
-#define VectorSet_W0( Vec )		_mm_and_ps( Vec, SSE_XYZ_MASK )
+#define VectorSet_W0( Vec )		_mm_and_ps( Vec, GlobalVectorConstants::XYZMask )
 
 /**
 * Loads XYZ and sets W=1
@@ -868,7 +858,7 @@ FORCEINLINE uint32 appCountTrailingZeros(uint32 Value)
  * @param VecW		Source register for ___W (note: the fourth component is used, not the first)
  * @return			VectorRegister(VecXYZ.x, VecXYZ.y, VecXYZ.z, VecW.w)
  */
-#define VectorMergeVecXYZ_VecW( VecXYZ, VecW )	VectorSelect( SSE_XYZ_MASK, VecXYZ, VecW )
+#define VectorMergeVecXYZ_VecW( VecXYZ, VecW )	VectorSelect( GlobalVectorConstants::XYZMask, VecXYZ, VecW )
 
 /**
  * Loads 4 BYTEs from unaligned memory and converts them into 4 FLOATs.
@@ -986,13 +976,140 @@ inline bool VectorContainsNaNOrInfinite(const VectorRegister& Vec)
 	// Test for infinity, technique "stolen" from DirectXMathVector.inl
 
 	// Mask off signs
-	VectorRegister InfTest = _mm_and_ps(Vec, SSE_SIGN_MASK);
+	VectorRegister InfTest = _mm_and_ps(Vec, GlobalVectorConstants::SignMask);
 	// Compare to infinity. If any are infinity, the signs are true.
-	bool IsNotInf = _mm_movemask_ps(_mm_cmpeq_ps(InfTest, SSE_INFINITY)) == 0;
+	bool IsNotInf = _mm_movemask_ps(_mm_cmpeq_ps(InfTest, GlobalVectorConstants::FloatInfinity)) == 0;
 
 	return !(IsNotNAN & IsNotInf);
 }
 
+FORCEINLINE VectorRegister VectorTruncate(const VectorRegister& X)
+{
+	return _mm_cvtepi32_ps(_mm_cvttps_epi32(X));
+}
+
+FORCEINLINE VectorRegister VectorFractional(const VectorRegister& X)
+{
+	return VectorSubtract(X, VectorTruncate(X));
+}
+
+FORCEINLINE VectorRegister VectorCeil(const VectorRegister& X)
+{
+	VectorRegister Trunc = VectorTruncate(X);
+	VectorRegister PosMask = VectorCompareGE(X, GlobalVectorConstants::FloatZero);
+	VectorRegister Add = VectorSelect(PosMask, GlobalVectorConstants::FloatOne, (GlobalVectorConstants::FloatZero));
+	return VectorAdd(Trunc, Add);
+}
+
+FORCEINLINE VectorRegister VectorFloor(const VectorRegister& X)
+{
+	VectorRegister Trunc = VectorTruncate(X);
+	VectorRegister PosMask = VectorCompareGE(X, (GlobalVectorConstants::FloatZero));
+	VectorRegister Sub = VectorSelect(PosMask, (GlobalVectorConstants::FloatZero), (GlobalVectorConstants::FloatOne));
+	return VectorSubtract(Trunc, Sub);
+}
+
+FORCEINLINE VectorRegister VectorMod(const VectorRegister& X, const VectorRegister& Y)
+{
+	VectorRegister Temp = VectorFloor(VectorDivide(X, Y));
+	return VectorSubtract(X, VectorMultiply(Y, Temp));
+}
+
+FORCEINLINE VectorRegister VectorSign(const VectorRegister& X)
+{
+	VectorRegister Mask = VectorCompareGE(X, (GlobalVectorConstants::FloatZero));
+	return VectorSelect(Mask, (GlobalVectorConstants::FloatOne), (GlobalVectorConstants::FloatMinusOne));
+}
+
+FORCEINLINE VectorRegister VectorStep(const VectorRegister& X)
+{
+	VectorRegister Mask = VectorCompareGE(X, (GlobalVectorConstants::FloatZero));
+	return VectorSelect(Mask, (GlobalVectorConstants::FloatOne), (GlobalVectorConstants::FloatZero));
+}
+
+//TODO: Vectorize
+FORCEINLINE VectorRegister VectorExp(const VectorRegister& X)
+{
+	return MakeVectorRegister((float)exp(VectorGetComponent(X, 0)), (float)exp(VectorGetComponent(X, 1)), (float)exp(VectorGetComponent(X, 2)), (float)exp(VectorGetComponent(X, 3)));
+}
+
+//TODO: Vectorize
+FORCEINLINE VectorRegister VectorExp2(const VectorRegister& X)
+{
+	return MakeVectorRegister((float)exp2(VectorGetComponent(X, 0)), (float)exp2(VectorGetComponent(X, 1)), (float)exp2(VectorGetComponent(X, 2)), (float)exp2(VectorGetComponent(X, 3)));
+}
+
+//TODO: Vectorize
+FORCEINLINE VectorRegister VectorLog(const VectorRegister& X)
+{
+	return MakeVectorRegister((float)log(VectorGetComponent(X, 0)), (float)log(VectorGetComponent(X, 1)), (float)log(VectorGetComponent(X, 2)), (float)log(VectorGetComponent(X, 3)));
+}
+
+//TODO: Vectorize
+FORCEINLINE VectorRegister VectorLog2(const VectorRegister& X)
+{
+	return MakeVectorRegister((float)log2(VectorGetComponent(X, 0)), (float)log2(VectorGetComponent(X, 1)), (float)log2(VectorGetComponent(X, 2)), (float)log2(VectorGetComponent(X, 3)));
+}
+
+FORCEINLINE VectorRegister VectorSin(const VectorRegister& X)
+{
+	//Sine approximation using a squared parabola restrained to f(0) = 0, f(PI) = 0, f(PI/2) = 1.
+	//based on a good discussion here http://forum.devmaster.net/t/fast-and-accurate-sine-cosine/9648
+	//After approx 2.5 million tests comparing to sin(): 
+	//Average error of 0.000128
+	//Max error of 0.001091
+
+	static const float p = 0.225f;
+	static const VectorRegister P = MakeVectorRegister(p, p, p, p);
+	static const float a = 16 * sqrt(p);
+	static const VectorRegister A = MakeVectorRegister(a, a, a, a);
+	static const float b = (1 - p) / sqrt(p);
+	static const VectorRegister B = MakeVectorRegister(b, b, b, b);
+
+	VectorRegister y = VectorMultiply(X, GlobalVectorConstants::OneOverTwoPi);
+	y = VectorSubtract(y, VectorFloor(VectorAdd(y, GlobalVectorConstants::FloatOneHalf)));
+	y = VectorMultiply(A, VectorMultiply(y, VectorSubtract(GlobalVectorConstants::FloatOneHalf, VectorAbs(y))));
+	return VectorMultiply(y, VectorAdd(B, VectorAbs(y)));
+}
+
+FORCEINLINE VectorRegister VectorCos(const VectorRegister& X)
+{
+	return VectorSin(VectorAdd(X, GlobalVectorConstants::PiByTwo));
+}
+
+//TODO: Vectorize
+FORCEINLINE VectorRegister VectorTan(const VectorRegister& X)
+{
+	return MakeVectorRegister((float)tan(VectorGetComponent(X, 0)), (float)tan(VectorGetComponent(X, 1)), (float)tan(VectorGetComponent(X, 2)), (float)tan(VectorGetComponent(X, 3)));
+}
+
+//TODO: Vectorize
+FORCEINLINE VectorRegister VectorASin(const VectorRegister& X)
+{
+	return MakeVectorRegister((float)asin(VectorGetComponent(X, 0)), (float)asin(VectorGetComponent(X, 1)), (float)asin(VectorGetComponent(X, 2)), (float)asin(VectorGetComponent(X, 3)));
+}
+
+//TODO: Vectorize
+FORCEINLINE VectorRegister VectorACos(const VectorRegister& X)
+{
+	return MakeVectorRegister((float)acos(VectorGetComponent(X, 0)), (float)acos(VectorGetComponent(X, 1)), (float)acos(VectorGetComponent(X, 2)), (float)acos(VectorGetComponent(X, 3)));
+}
+
+//TODO: Vectorize
+FORCEINLINE VectorRegister VectorATan(const VectorRegister& X)
+{
+	return MakeVectorRegister((float)atan(VectorGetComponent(X, 0)), (float)atan(VectorGetComponent(X, 1)), (float)atan(VectorGetComponent(X, 2)), (float)atan(VectorGetComponent(X, 3)));
+}
+
+//TODO: Vectorize
+FORCEINLINE VectorRegister VectorATan2(const VectorRegister& X, const VectorRegister& Y)
+{
+	return MakeVectorRegister(	(float)atan2(VectorGetComponent(X, 0), VectorGetComponent(Y, 0)), 
+		(float)atan2(VectorGetComponent(X, 1), VectorGetComponent(Y, 1)),
+		(float)atan2(VectorGetComponent(X, 2), VectorGetComponent(Y, 2)),
+		(float)atan2(VectorGetComponent(X, 3), VectorGetComponent(Y, 3)));
+}
+
 // To be continued...
 
 #endif
diff --git a/Engine/Source/Runtime/Core/Public/Math/UnrealMathVectorConstants.h b/Engine/Source/Runtime/Core/Public/Math/UnrealMathVectorConstants.h
--- a/Engine/Source/Runtime/Core/Public/Math/UnrealMathVectorConstants.h
+++ b/Engine/Source/Runtime/Core/Public/Math/UnrealMathVectorConstants.h
@@ -5,6 +5,9 @@
 
 namespace GlobalVectorConstants
 {
+	static const VectorRegister FloatOne = MakeVectorRegister(1.0f, 1.0f, 1.0f, 1.0f);
+	static const VectorRegister FloatZero = MakeVectorRegister(0.0f, 0.0f, 0.0f, 0.0f);
+	static const VectorRegister FloatMinusOne = MakeVectorRegister(-1.0f, -1.0f, -1.0f, -1.0f);
 	static const VectorRegister Float0001 = MakeVectorRegister( 0.0f, 0.0f, 0.0f, 1.0f );
 	static const VectorRegister SmallLengthThreshold = MakeVectorRegister(1.e-8f, 1.e-8f, 1.e-8f, 1.e-8f);
 	static const VectorRegister FloatOneHundredth = MakeVectorRegister(0.01f, 0.01f, 0.01f, 0.01f);
@@ -21,4 +24,27 @@ namespace GlobalVectorConstants
 	static const VectorRegister QMULTI_SIGN_MASK0 = MakeVectorRegister( 1.f, -1.f, 1.f, -1.f );
 	static const VectorRegister QMULTI_SIGN_MASK1 = MakeVectorRegister( 1.f, 1.f, -1.f, -1.f );
 	static const VectorRegister QMULTI_SIGN_MASK2 = MakeVectorRegister( -1.f, 1.f, 1.f, -1.f );
+
+
+	/** Bitmask to AND out the XYZ components in a vector */
+	static const VectorRegister XYZMask = MakeVectorRegister((uint32)0xffffffff, (uint32)0xffffffff, (uint32)0xffffffff, (uint32)0x00000000);
+
+	/** Bitmask to AND out the sign bit of each components in a vector */
+#define SIGN_BIT ((1 << 31))
+	static const VectorRegister SignBit = MakeVectorRegister((uint32)SIGN_BIT, (uint32)SIGN_BIT, (uint32)SIGN_BIT, (uint32)SIGN_BIT);
+	static const VectorRegister SignMask = MakeVectorRegister((uint32)(~SIGN_BIT), (uint32)(~SIGN_BIT), (uint32)(~SIGN_BIT), (uint32)(~SIGN_BIT));
+#undef SIGN_BIT
+
+	/** Vector full of positive infinity */
+	static const VectorRegister FloatInfinity = MakeVectorRegister((uint32)0x7F800000, (uint32)0x7F800000, (uint32)0x7F800000, (uint32)0x7F800000);
+
+
+	static const VectorRegister Pi = MakeVectorRegister(PI, PI, PI, PI);
+	static const VectorRegister TwoPi = MakeVectorRegister(2.0f*PI, 2.0f*PI, 2.0f*PI, 2.0f*PI);
+	static const VectorRegister PiByTwo = MakeVectorRegister(0.5f*PI, 0.5f*PI, 0.5f*PI, 0.5f*PI);
+	static const VectorRegister PiByFour = MakeVectorRegister(0.25f*PI, 0.25f*PI, 0.25f*PI, 0.25f*PI);
+	static const VectorRegister OneOverPi = MakeVectorRegister(1.0f / PI, 1.0f / PI, 1.0f / PI, 1.0f / PI);
+	static const VectorRegister OneOverTwoPi = MakeVectorRegister(1.0f / (2.0f*PI), 1.0f / (2.0f*PI), 1.0f / (2.0f*PI), 1.0f / (2.0f*PI));
+
+	static const VectorRegister Float255 = MakeVectorRegister(255.0f, 255.0f, 255.0f, 255.0f);
 }
diff --git a/Engine/Source/Runtime/VectorVM/Private/VectorVM.cpp b/Engine/Source/Runtime/VectorVM/Private/VectorVM.cpp
--- a/Engine/Source/Runtime/VectorVM/Private/VectorVM.cpp
+++ b/Engine/Source/Runtime/VectorVM/Private/VectorVM.cpp
@@ -8,6 +8,9 @@ IMPLEMENT_MODULE(FDefaultModuleImpl, VectorVM);
 
 DEFINE_LOG_CATEGORY_STATIC(LogVectorVM, All, All);
 
+//#define VM_FORCEINLINE 
+#define VM_FORCEINLINE FORCEINLINE
+
 #define SRCOP_RRRR 0x00
 #define SRCOP_RRRC 0x01
 #define SRCOP_RRCR 0x02
@@ -55,30 +58,30 @@ struct FVectorVMContext
 };
 
 /** Decode the next operation contained in the bytecode. */
-static FORCEINLINE VectorVM::EOp DecodeOp(FVectorVMContext& Context)
+static VM_FORCEINLINE VectorVM::EOp DecodeOp(FVectorVMContext& Context)
 {
 	return static_cast<VectorVM::EOp>(*Context.Code++);
 }
 
 /** Decode a register from the bytecode. */
-static FORCEINLINE VectorRegister* DecodeRegister(FVectorVMContext& Context)
+static VM_FORCEINLINE VectorRegister* DecodeRegister(FVectorVMContext& Context)
 {
 	return Context.RegisterTable[*Context.Code++];
 }
 
 /** Decode a constant from the bytecode. */
-static FORCEINLINE VectorRegister DecodeConstant(FVectorVMContext& Context)
+static VM_FORCEINLINE VectorRegister DecodeConstant(FVectorVMContext& Context)
 {
 	const FVector4* vec = &Context.ConstantTable[*Context.Code++];
 	return VectorLoad(vec);
 }
 
-static FORCEINLINE uint8 DecodeSrcOperandTypes(FVectorVMContext& Context)
+static VM_FORCEINLINE uint8 DecodeSrcOperandTypes(FVectorVMContext& Context)
 {
 	return *Context.Code++;
 }
 
-static FORCEINLINE uint8 DecodeByte(FVectorVMContext& Context)
+static VM_FORCEINLINE uint8 DecodeByte(FVectorVMContext& Context)
 {
 	return *Context.Code++;
 }
@@ -90,7 +93,7 @@ struct FConstantHandler
 	FConstantHandler(FVectorVMContext& Context)
 		: Constant(DecodeConstant(Context))
 	{}
-	FORCEINLINE const VectorRegister& Get(){ return Constant; }
+	VM_FORCEINLINE const VectorRegister& Get(){ return Constant; }
 };
 
 /** Handles reading of a register, advancing the pointer with each read. */
@@ -101,11 +104,11 @@ struct FRegisterHandler
 	FRegisterHandler(FVectorVMContext& Context)
 		: Register(DecodeRegister(Context))
 	{}
-	FORCEINLINE const VectorRegister& Get(){ RegisterValue = *Register++;  return RegisterValue; }
+	VM_FORCEINLINE const VectorRegister& Get(){ RegisterValue = *Register++;  return RegisterValue; }
 };
 
 template<typename Kernel, typename Arg0Handler>
-FORCEINLINE void VectorUnaryLoop(FVectorVMContext& Context, VectorRegister* RESTRICT Dst, int32 NumVectors)
+VM_FORCEINLINE void VectorUnaryLoop(FVectorVMContext& Context, VectorRegister* RESTRICT Dst, int32 NumVectors)
 {
 	Arg0Handler Arg0(Context);
 	for (int32 i = 0; i < NumVectors; ++i)
@@ -115,7 +118,7 @@ FORCEINLINE void VectorUnaryLoop(FVectorVMContext& Context, VectorRegister* REST
 }
 
 template<typename Kernel, typename Arg0Handler, typename Arg1Handler>
-FORCEINLINE void VectorBinaryLoop(FVectorVMContext& Context, VectorRegister* RESTRICT Dst, int32 NumVectors)
+VM_FORCEINLINE void VectorBinaryLoop(FVectorVMContext& Context, VectorRegister* RESTRICT Dst, int32 NumVectors)
 {
 	Arg0Handler Arg0(Context);
 	Arg1Handler Arg1(Context);
@@ -126,7 +129,7 @@ FORCEINLINE void VectorBinaryLoop(FVectorVMContext& Context, VectorRegister* RES
 }
 
 template<typename Kernel, typename Arg0Handler, typename Arg1Handler, typename Arg2Handler>
-FORCEINLINE void VectorTrinaryLoop(FVectorVMContext& Context, VectorRegister* RESTRICT Dst, int32 NumVectors)
+VM_FORCEINLINE void VectorTrinaryLoop(FVectorVMContext& Context, VectorRegister* RESTRICT Dst, int32 NumVectors)
 {
 	Arg0Handler Arg0(Context);
 	Arg1Handler Arg1(Context);
@@ -138,7 +141,7 @@ FORCEINLINE void VectorTrinaryLoop(FVectorVMContext& Context, VectorRegister* RE
 }
 
 template<typename Kernel, typename Arg0Handler, typename Arg1Handler, typename Arg2Handler, typename Arg3Handler>
-FORCEINLINE void VectorQuinaryLoop(FVectorVMContext& Context, VectorRegister* RESTRICT Dst, int32 NumVectors)
+VM_FORCEINLINE void VectorQuatenaryLoop(FVectorVMContext& Context, VectorRegister* RESTRICT Dst, int32 NumVectors)
 {
 	Arg0Handler Arg0(Context);
 	Arg1Handler Arg1(Context);
@@ -214,7 +217,7 @@ struct TTrinaryVectorKernel
 
 /** Base class of Vector kernels with 4 operands. */
 template <typename Kernel>
-struct TQuinaryVectorKernel
+struct TQuatenaryVectorKernel
 {
 	static void Exec(FVectorVMContext& Context)
 	{
@@ -223,22 +226,22 @@ struct TQuinaryVectorKernel
 		int32 NumVectors = Context.NumVectors;
 		switch (SrcOpTypes)
 		{
-		case SRCOP_RRRR: 	VectorQuinaryLoop<Kernel, FRegisterHandler, FRegisterHandler, FRegisterHandler, FRegisterHandler>(Context, Dst, NumVectors); break;
-		case SRCOP_RRRC:	VectorQuinaryLoop<Kernel, FConstantHandler, FRegisterHandler, FRegisterHandler, FRegisterHandler>(Context, Dst, NumVectors); break;
-		case SRCOP_RRCR: 	VectorQuinaryLoop<Kernel, FRegisterHandler, FConstantHandler, FRegisterHandler, FRegisterHandler>(Context, Dst, NumVectors); break;
-		case SRCOP_RRCC:	VectorQuinaryLoop<Kernel, FConstantHandler, FConstantHandler, FRegisterHandler, FRegisterHandler>(Context, Dst, NumVectors); break;
-		case SRCOP_RCRR: 	VectorQuinaryLoop<Kernel, FRegisterHandler, FRegisterHandler, FConstantHandler, FRegisterHandler>(Context, Dst, NumVectors); break;
-		case SRCOP_RCRC:	VectorQuinaryLoop<Kernel, FConstantHandler, FRegisterHandler, FConstantHandler, FRegisterHandler>(Context, Dst, NumVectors); break;
-		case SRCOP_RCCR: 	VectorQuinaryLoop<Kernel, FRegisterHandler, FConstantHandler, FConstantHandler, FRegisterHandler>(Context, Dst, NumVectors); break;
-		case SRCOP_RCCC:	VectorQuinaryLoop<Kernel, FConstantHandler, FConstantHandler, FConstantHandler, FRegisterHandler>(Context, Dst, NumVectors); break;
-		case SRCOP_CRRR: 	VectorQuinaryLoop<Kernel, FRegisterHandler, FRegisterHandler, FRegisterHandler, FConstantHandler>(Context, Dst, NumVectors); break;
-		case SRCOP_CRRC:	VectorQuinaryLoop<Kernel, FConstantHandler, FRegisterHandler, FRegisterHandler, FConstantHandler>(Context, Dst, NumVectors); break;
-		case SRCOP_CRCR: 	VectorQuinaryLoop<Kernel, FRegisterHandler, FConstantHandler, FRegisterHandler, FConstantHandler>(Context, Dst, NumVectors); break;
-		case SRCOP_CRCC:	VectorQuinaryLoop<Kernel, FConstantHandler, FConstantHandler, FRegisterHandler, FConstantHandler>(Context, Dst, NumVectors); break;
-		case SRCOP_CCRR: 	VectorQuinaryLoop<Kernel, FRegisterHandler, FRegisterHandler, FConstantHandler, FConstantHandler>(Context, Dst, NumVectors); break;
-		case SRCOP_CCRC:	VectorQuinaryLoop<Kernel, FConstantHandler, FRegisterHandler, FConstantHandler, FConstantHandler>(Context, Dst, NumVectors); break;
-		case SRCOP_CCCR: 	VectorQuinaryLoop<Kernel, FRegisterHandler, FConstantHandler, FConstantHandler, FConstantHandler>(Context, Dst, NumVectors); break;
-		case SRCOP_CCCC:	VectorQuinaryLoop<Kernel, FConstantHandler, FConstantHandler, FConstantHandler, FConstantHandler>(Context, Dst, NumVectors); break;
+		case SRCOP_RRRR: 	VectorQuatenaryLoop<Kernel, FRegisterHandler, FRegisterHandler, FRegisterHandler, FRegisterHandler>(Context, Dst, NumVectors); break;
+		case SRCOP_RRRC:	VectorQuatenaryLoop<Kernel, FConstantHandler, FRegisterHandler, FRegisterHandler, FRegisterHandler>(Context, Dst, NumVectors); break;
+		case SRCOP_RRCR: 	VectorQuatenaryLoop<Kernel, FRegisterHandler, FConstantHandler, FRegisterHandler, FRegisterHandler>(Context, Dst, NumVectors); break;
+		case SRCOP_RRCC:	VectorQuatenaryLoop<Kernel, FConstantHandler, FConstantHandler, FRegisterHandler, FRegisterHandler>(Context, Dst, NumVectors); break;
+		case SRCOP_RCRR: 	VectorQuatenaryLoop<Kernel, FRegisterHandler, FRegisterHandler, FConstantHandler, FRegisterHandler>(Context, Dst, NumVectors); break;
+		case SRCOP_RCRC:	VectorQuatenaryLoop<Kernel, FConstantHandler, FRegisterHandler, FConstantHandler, FRegisterHandler>(Context, Dst, NumVectors); break;
+		case SRCOP_RCCR: 	VectorQuatenaryLoop<Kernel, FRegisterHandler, FConstantHandler, FConstantHandler, FRegisterHandler>(Context, Dst, NumVectors); break;
+		case SRCOP_RCCC:	VectorQuatenaryLoop<Kernel, FConstantHandler, FConstantHandler, FConstantHandler, FRegisterHandler>(Context, Dst, NumVectors); break;
+		case SRCOP_CRRR: 	VectorQuatenaryLoop<Kernel, FRegisterHandler, FRegisterHandler, FRegisterHandler, FConstantHandler>(Context, Dst, NumVectors); break;
+		case SRCOP_CRRC:	VectorQuatenaryLoop<Kernel, FConstantHandler, FRegisterHandler, FRegisterHandler, FConstantHandler>(Context, Dst, NumVectors); break;
+		case SRCOP_CRCR: 	VectorQuatenaryLoop<Kernel, FRegisterHandler, FConstantHandler, FRegisterHandler, FConstantHandler>(Context, Dst, NumVectors); break;
+		case SRCOP_CRCC:	VectorQuatenaryLoop<Kernel, FConstantHandler, FConstantHandler, FRegisterHandler, FConstantHandler>(Context, Dst, NumVectors); break;
+		case SRCOP_CCRR: 	VectorQuatenaryLoop<Kernel, FRegisterHandler, FRegisterHandler, FConstantHandler, FConstantHandler>(Context, Dst, NumVectors); break;
+		case SRCOP_CCRC:	VectorQuatenaryLoop<Kernel, FConstantHandler, FRegisterHandler, FConstantHandler, FConstantHandler>(Context, Dst, NumVectors); break;
+		case SRCOP_CCCR: 	VectorQuatenaryLoop<Kernel, FRegisterHandler, FConstantHandler, FConstantHandler, FConstantHandler>(Context, Dst, NumVectors); break;
+		case SRCOP_CCCC:	VectorQuatenaryLoop<Kernel, FConstantHandler, FConstantHandler, FConstantHandler, FConstantHandler>(Context, Dst, NumVectors); break;
 		default: check(0); break;
 		};
 	}
@@ -250,7 +253,7 @@ struct TQuinaryVectorKernel
 
 struct FVectorKernelAdd : public TBinaryVectorKernel<FVectorKernelAdd>
 {
-	static void FORCEINLINE DoKernel(VectorRegister* RESTRICT Dst,VectorRegister Src0,VectorRegister Src1)
+	static void VM_FORCEINLINE DoKernel(VectorRegister* RESTRICT Dst,VectorRegister Src0,VectorRegister Src1)
 	{
 		*Dst = VectorAdd(Src0, Src1);
 	}
@@ -258,7 +261,7 @@ struct FVectorKernelAdd : public TBinaryVectorKernel<FVectorKernelAdd>
 
 struct FVectorKernelSub : public TBinaryVectorKernel<FVectorKernelSub>
 {
-	static void FORCEINLINE DoKernel(VectorRegister* RESTRICT Dst,VectorRegister Src0,VectorRegister Src1)
+	static void VM_FORCEINLINE DoKernel(VectorRegister* RESTRICT Dst,VectorRegister Src0,VectorRegister Src1)
 	{
 		*Dst = VectorSubtract(Src0, Src1);
 	}
@@ -266,7 +269,7 @@ struct FVectorKernelSub : public TBinaryVectorKernel<FVectorKernelSub>
 
 struct FVectorKernelMul : public TBinaryVectorKernel<FVectorKernelMul>
 {
-	static void FORCEINLINE DoKernel(VectorRegister* RESTRICT Dst,VectorRegister Src0,VectorRegister Src1)
+	static void VM_FORCEINLINE DoKernel(VectorRegister* RESTRICT Dst,VectorRegister Src0,VectorRegister Src1)
 	{
 		*Dst = VectorMultiply(Src0, Src1);
 	}
@@ -274,7 +277,7 @@ struct FVectorKernelMul : public TBinaryVectorKernel<FVectorKernelMul>
 
 struct FVectorKernelMad : public TTrinaryVectorKernel<FVectorKernelMad>
 {
-	static void FORCEINLINE DoKernel(VectorRegister* RESTRICT Dst,VectorRegister Src0,VectorRegister Src1,VectorRegister Src2)
+	static void VM_FORCEINLINE DoKernel(VectorRegister* RESTRICT Dst,VectorRegister Src0,VectorRegister Src1,VectorRegister Src2)
 	{
 		*Dst = VectorMultiplyAdd(Src0, Src1, Src2);
 	}
@@ -282,7 +285,7 @@ struct FVectorKernelMad : public TTrinaryVectorKernel<FVectorKernelMad>
 
 struct FVectorKernelLerp : public TTrinaryVectorKernel<FVectorKernelLerp>
 {
-	static void FORCEINLINE DoKernel(VectorRegister* RESTRICT Dst,VectorRegister Src0,VectorRegister Src1,VectorRegister Src2)
+	static void VM_FORCEINLINE DoKernel(VectorRegister* RESTRICT Dst,VectorRegister Src0,VectorRegister Src1,VectorRegister Src2)
 	{
 		const VectorRegister One = MakeVectorRegister(1.0f, 1.0f, 1.0f, 1.0f);
 		const VectorRegister OneMinusAlpha = VectorSubtract(One, Src2);
@@ -293,7 +296,7 @@ struct FVectorKernelLerp : public TTrinaryVectorKernel<FVectorKernelLerp>
 
 struct FVectorKernelRcp : public TUnaryVectorKernel<FVectorKernelRcp>
 {
-	static void FORCEINLINE DoKernel(VectorRegister* RESTRICT Dst,VectorRegister Src0)
+	static void VM_FORCEINLINE DoKernel(VectorRegister* RESTRICT Dst,VectorRegister Src0)
 	{
 		*Dst = VectorReciprocal(Src0);
 	}
@@ -301,7 +304,7 @@ struct FVectorKernelRcp : public TUnaryVectorKernel<FVectorKernelRcp>
 
 struct FVectorKernelRsq : public TUnaryVectorKernel<FVectorKernelRsq>
 {
-	static void FORCEINLINE DoKernel(VectorRegister* RESTRICT Dst,VectorRegister Src0)
+	static void VM_FORCEINLINE DoKernel(VectorRegister* RESTRICT Dst,VectorRegister Src0)
 	{
 		*Dst = VectorReciprocalSqrt(Src0);
 	}
@@ -309,7 +312,7 @@ struct FVectorKernelRsq : public TUnaryVectorKernel<FVectorKernelRsq>
 
 struct FVectorKernelSqrt : public TUnaryVectorKernel<FVectorKernelSqrt>
 {
-	static void FORCEINLINE DoKernel(VectorRegister* RESTRICT Dst,VectorRegister Src0)
+	static void VM_FORCEINLINE DoKernel(VectorRegister* RESTRICT Dst,VectorRegister Src0)
 	{
 		// TODO: Need a SIMD sqrt!
 		float* RESTRICT FloatDst = reinterpret_cast<float* RESTRICT>(Dst);
@@ -323,7 +326,7 @@ struct FVectorKernelSqrt : public TUnaryVectorKernel<FVectorKernelSqrt>
 
 struct FVectorKernelNeg : public TUnaryVectorKernel<FVectorKernelNeg>
 {
-	static void FORCEINLINE DoKernel(VectorRegister* RESTRICT Dst,VectorRegister Src0)
+	static void VM_FORCEINLINE DoKernel(VectorRegister* RESTRICT Dst,VectorRegister Src0)
 	{
 		*Dst = VectorNegate(Src0);
 	}
@@ -331,42 +334,146 @@ struct FVectorKernelNeg : public TUnaryVectorKernel<FVectorKernelNeg>
 
 struct FVectorKernelAbs : public TUnaryVectorKernel<FVectorKernelAbs>
 {
-	static void FORCEINLINE DoKernel(VectorRegister* RESTRICT Dst,VectorRegister Src0)
+	static void VM_FORCEINLINE DoKernel(VectorRegister* RESTRICT Dst,VectorRegister Src0)
 	{
 		*Dst = VectorAbs(Src0);
 	}
 };
 
+struct FVectorKernelExp : public TUnaryVectorKernel<FVectorKernelExp>
+{
+	static void VM_FORCEINLINE DoKernel(VectorRegister* RESTRICT Dst, VectorRegister Src0)
+	{
+		*Dst = VectorExp(Src0);
+	}
+}; 
+
+struct FVectorKernelExp2 : public TUnaryVectorKernel<FVectorKernelExp2>
+{
+	static void VM_FORCEINLINE DoKernel(VectorRegister* RESTRICT Dst, VectorRegister Src0)
+	{
+		*Dst = VectorExp2(Src0);
+	}
+};
+
+struct FVectorKernelLog : public TUnaryVectorKernel<FVectorKernelLog>
+{
+	static void VM_FORCEINLINE DoKernel(VectorRegister* RESTRICT Dst, VectorRegister Src0)
+	{
+		*Dst = VectorLog(Src0);
+	}
+};
+
+struct FVectorKernelLog2 : public TUnaryVectorKernel<FVectorKernelLog2>
+{
+	static void VM_FORCEINLINE DoKernel(VectorRegister* RESTRICT Dst, VectorRegister Src0)
+	{
+		*Dst = VectorLog2(Src0);
+	}
+};
+
 struct FVectorKernelClamp : public TTrinaryVectorKernel<FVectorKernelClamp>
 {
-	static void FORCEINLINE DoKernel(VectorRegister* RESTRICT Dst,VectorRegister Src0,VectorRegister Src1,VectorRegister Src2)
+	static void VM_FORCEINLINE DoKernel(VectorRegister* RESTRICT Dst,VectorRegister Src0,VectorRegister Src1,VectorRegister Src2)
 	{
 		const VectorRegister Tmp = VectorMax(Src0, Src1);
 		*Dst = VectorMin(Tmp, Src2);
 	}
 };
 
-
 struct FVectorKernelSin : public TUnaryVectorKernel<FVectorKernelSin>
 {
-	static void FORCEINLINE DoKernel(VectorRegister* RESTRICT Dst, VectorRegister Src0)
+	static void VM_FORCEINLINE DoKernel(VectorRegister* RESTRICT Dst, VectorRegister Src0)
 	{
-		float v = VectorGetComponent(Src0, 0);
-		float sn = FMath::Sin(v*3.14f);		// [0;1] takes us through half a period
-		*Dst = MakeVectorRegister(sn, sn, sn, sn);
+		*Dst = VectorSin(Src0);
 	}
 };
 
-struct FVectorKernelSin4 : public TUnaryVectorKernel<FVectorKernelSin4>
+struct FVectorKernelCos : public TUnaryVectorKernel<FVectorKernelCos>
 {
-	static void FORCEINLINE DoKernel(VectorRegister* RESTRICT Dst, VectorRegister Src0)
+	static void VM_FORCEINLINE DoKernel(VectorRegister* RESTRICT Dst, VectorRegister Src0)
+ 	{
+		*Dst = VectorCos(Src0);
+	}
+};
+
+struct FVectorKernelTan : public TUnaryVectorKernel<FVectorKernelTan>
+{
+	static void VM_FORCEINLINE DoKernel(VectorRegister* RESTRICT Dst, VectorRegister Src0)
 	{
-		float const* FloatSrc0 = reinterpret_cast<float const*>(&Src0);
-		float sn1 = FMath::Sin(FloatSrc0[0] * 3.14f);		// [0;1] takes us through half a period
-		float sn2 = FMath::Sin(FloatSrc0[1] * 3.14f);		// [0;1] takes us through half a period
-		float sn3 = FMath::Sin(FloatSrc0[2] * 3.14f);		// [0;1] takes us through half a period
-		float sn4 = FMath::Sin(FloatSrc0[3] * 3.14f);		// [0;1] takes us through half a period
-		*Dst = MakeVectorRegister(sn1, sn2, sn3, sn4);
+		*Dst = VectorTan(Src0);
+	}
+};
+
+struct FVectorKernelASin : public TUnaryVectorKernel<FVectorKernelASin>
+{
+	static void VM_FORCEINLINE DoKernel(VectorRegister* RESTRICT Dst, VectorRegister Src0)
+	{
+		*Dst = VectorASin(Src0);
+	}
+};
+
+struct FVectorKernelACos : public TUnaryVectorKernel<FVectorKernelACos>
+{
+	static void VM_FORCEINLINE DoKernel(VectorRegister* RESTRICT Dst, VectorRegister Src0)
+	{
+		*Dst = VectorACos(Src0);
+	}
+};
+
+struct FVectorKernelATan : public TUnaryVectorKernel<FVectorKernelATan>
+{
+	static void VM_FORCEINLINE DoKernel(VectorRegister* RESTRICT Dst, VectorRegister Src0)
+	{
+		*Dst = VectorATan(Src0);
+	}
+};
+
+struct FVectorKernelATan2 : public TBinaryVectorKernel<FVectorKernelATan2>
+{
+	static void VM_FORCEINLINE DoKernel(VectorRegister* RESTRICT Dst, VectorRegister Src0, VectorRegister Src1)
+	{
+		*Dst = VectorATan2(Src0, Src1);
+	}
+};
+
+struct FVectorKernelCeil : public TUnaryVectorKernel<FVectorKernelCeil>
+{
+	static void VM_FORCEINLINE DoKernel(VectorRegister* RESTRICT Dst, VectorRegister Src0)
+	{
+		*Dst = VectorCeil(Src0);
+	}
+};
+
+struct FVectorKernelFloor : public TUnaryVectorKernel<FVectorKernelFloor>
+{
+	static void VM_FORCEINLINE DoKernel(VectorRegister* RESTRICT Dst, VectorRegister Src0)
+	{
+		*Dst = VectorFloor(Src0);
+	}
+};
+
+struct FVectorKernelMod : public TBinaryVectorKernel<FVectorKernelMod>
+{
+	static void VM_FORCEINLINE DoKernel(VectorRegister* RESTRICT Dst, VectorRegister Src0, VectorRegister Src1)
+	{
+		*Dst = VectorMod(Src0, Src1);
+	}
+};
+
+struct FVectorKernelFrac : public TUnaryVectorKernel<FVectorKernelFrac>
+{
+	static void VM_FORCEINLINE DoKernel(VectorRegister* RESTRICT Dst, VectorRegister Src0)
+	{
+		*Dst = VectorFractional(Src0);
+	}
+};
+
+struct FVectorKernelTrunc : public TUnaryVectorKernel<FVectorKernelTrunc>
+{
+	static void VM_FORCEINLINE DoKernel(VectorRegister* RESTRICT Dst, VectorRegister Src0)
+	{
+		*Dst = VectorTruncate(Src0);
 	}
 };
 
@@ -396,16 +503,15 @@ struct FVectorKernelLessThan : public TBinaryVectorKernel<FVectorKernelLessThan>
 
 struct FVectorKernelDot : public TBinaryVectorKernel<FVectorKernelDot>
 {
-	static void FORCEINLINE DoKernel(VectorRegister* RESTRICT Dst, VectorRegister Src0, VectorRegister Src1)
+	static void VM_FORCEINLINE DoKernel(VectorRegister* RESTRICT Dst, VectorRegister Src0, VectorRegister Src1)
 	{
 		*Dst = VectorDot4(Src0, Src1);
 	}
 };
 
-
 struct FVectorKernelLength : public TUnaryVectorKernel<FVectorKernelLength>
 {
-	static void FORCEINLINE DoKernel(VectorRegister* RESTRICT Dst, VectorRegister Src0)
+	static void VM_FORCEINLINE DoKernel(VectorRegister* RESTRICT Dst, VectorRegister Src0)
 	{
 		VectorRegister Temp = VectorDot4(Src0, Src0);
 		float const* FloatSrc = reinterpret_cast<float const*>(&Temp);
@@ -414,10 +520,9 @@ struct FVectorKernelLength : public TUnaryVectorKernel<FVectorKernelLength>
 	}
 };
 
-
 struct FVectorKernelCross : public TBinaryVectorKernel<FVectorKernelCross>
 {
-	static void FORCEINLINE DoKernel(VectorRegister* RESTRICT Dst, VectorRegister Src0, VectorRegister Src1)
+	static void VM_FORCEINLINE DoKernel(VectorRegister* RESTRICT Dst, VectorRegister Src0, VectorRegister Src1)
 	{
 		*Dst = VectorCross(Src0, Src1);
 	}
@@ -425,7 +530,7 @@ struct FVectorKernelCross : public TBinaryVectorKernel<FVectorKernelCross>
 
 struct FVectorKernelNormalize : public TUnaryVectorKernel<FVectorKernelNormalize>
 {
-	static void FORCEINLINE DoKernel(VectorRegister* RESTRICT Dst, VectorRegister Src0)
+	static void VM_FORCEINLINE DoKernel(VectorRegister* RESTRICT Dst, VectorRegister Src0)
 	{
 		*Dst = VectorNormalize(Src0);
 	}
@@ -433,7 +538,7 @@ struct FVectorKernelNormalize : public TUnaryVectorKernel<FVectorKernelNormalize
 
 struct FVectorKernelRandom : public TUnaryVectorKernel<FVectorKernelRandom>
 {
-	static void FORCEINLINE DoKernel(VectorRegister* RESTRICT Dst, VectorRegister Src0)
+	static void VM_FORCEINLINE DoKernel(VectorRegister* RESTRICT Dst, VectorRegister Src0)
 	{
 		const float rm = RAND_MAX;
 		VectorRegister Result = MakeVectorRegister(static_cast<float>(FMath::Rand()) / rm,
@@ -449,7 +554,7 @@ struct FVectorKernelRandom : public TUnaryVectorKernel<FVectorKernelRandom>
 
 struct FVectorKernelMin : public TBinaryVectorKernel<FVectorKernelMin>
 {
-	static void FORCEINLINE DoKernel(VectorRegister* RESTRICT Dst,VectorRegister Src0,VectorRegister Src1)
+	static void VM_FORCEINLINE DoKernel(VectorRegister* RESTRICT Dst,VectorRegister Src0,VectorRegister Src1)
 	{
 		*Dst = VectorMin(Src0, Src1);
 	}
@@ -457,7 +562,7 @@ struct FVectorKernelMin : public TBinaryVectorKernel<FVectorKernelMin>
 
 struct FVectorKernelMax : public TBinaryVectorKernel<FVectorKernelMax>
 {
-	static void FORCEINLINE DoKernel(VectorRegister* RESTRICT Dst,VectorRegister Src0,VectorRegister Src1)
+	static void VM_FORCEINLINE DoKernel(VectorRegister* RESTRICT Dst,VectorRegister Src0,VectorRegister Src1)
 	{
 		*Dst = VectorMax(Src0, Src1);
 	}
@@ -465,18 +570,33 @@ struct FVectorKernelMax : public TBinaryVectorKernel<FVectorKernelMax>
 
 struct FVectorKernelPow : public TBinaryVectorKernel<FVectorKernelPow>
 {
-	static void FORCEINLINE DoKernel(VectorRegister* RESTRICT Dst,VectorRegister Src0,VectorRegister Src1)
+	static void VM_FORCEINLINE DoKernel(VectorRegister* RESTRICT Dst,VectorRegister Src0,VectorRegister Src1)
 	{
 		*Dst = VectorPow(Src0, Src1);
 	}
 };
 
+struct FVectorKernelSign : public TUnaryVectorKernel<FVectorKernelSign>
+{
+	static void VM_FORCEINLINE DoKernel(VectorRegister* RESTRICT Dst, VectorRegister Src0)
+	{
+		*Dst = VectorSign(Src0);
+	}
+};
+
+struct FVectorKernelStep : public TUnaryVectorKernel<FVectorKernelStep>
+{
+	static void VM_FORCEINLINE DoKernel(VectorRegister* RESTRICT Dst, VectorRegister Src0)
+	{
+		*Dst = VectorStep(Src0);
+	}
+};
 
 struct FVectorKernelNoise : public TUnaryVectorKernel<FVectorKernelNoise>
 {
 	static VectorRegister RandomTable[10][10][10];
 
-	static void FORCEINLINE DoKernel(VectorRegister* RESTRICT Dst, VectorRegister Src0)
+	static void VM_FORCEINLINE DoKernel(VectorRegister* RESTRICT Dst, VectorRegister Src0)
 	{
 		const VectorRegister One = MakeVectorRegister(1.0f, 1.0f, 1.0f, 1.0f);
 		const VectorRegister OneHalf = MakeVectorRegister(0.5f, 0.5f, 0.5f, 0.5f);
@@ -535,19 +655,19 @@ VectorRegister FVectorKernelNoise::RandomTable[10][10][10];
 template<int32 Component>
 struct FVectorKernelSplat : public TUnaryVectorKernel<FVectorKernelSplat<Component>>
 {
-	static void FORCEINLINE DoKernel(VectorRegister* RESTRICT Dst, VectorRegister Src0)
+	static void VM_FORCEINLINE DoKernel(VectorRegister* RESTRICT Dst, VectorRegister Src0)
 	{
 		*Dst = VectorReplicate(Src0, Component);
 	}
 };
 
 template<int32 Cmp0, int32 Cmp1, int32 Cmp2, int32 Cmp3>
-struct FVectorKernelCompose : public TQuinaryVectorKernel<FVectorKernelCompose<Cmp0, Cmp1, Cmp2, Cmp3>>
+struct FVectorKernelCompose : public TQuatenaryVectorKernel<FVectorKernelCompose<Cmp0, Cmp1, Cmp2, Cmp3>>
 {
 	//Passing as const refs as some compilers cant handle > 3 aligned vectorregister params.
 	//inlined so shouldn't impact perf
 	//Todo: ^^^^ test this
-	static void FORCEINLINE DoKernel(VectorRegister* RESTRICT Dst, const VectorRegister& Src0, const VectorRegister& Src1, const VectorRegister& Src2, const VectorRegister& Src3)
+	static void VM_FORCEINLINE DoKernel(VectorRegister* RESTRICT Dst, const VectorRegister& Src0, const VectorRegister& Src1, const VectorRegister& Src2, const VectorRegister& Src3)
 	{
 		//TODO - There's probably a faster way to do this.
 		VectorRegister Tmp0 = VectorShuffle(Src0, Src1, Cmp0, Cmp0, Cmp1, Cmp1);
@@ -558,7 +678,7 @@ struct FVectorKernelCompose : public TQuinaryVectorKernel<FVectorKernelCompose<C
 
 struct FVectorKernelOutput : public TUnaryVectorKernel<FVectorKernelOutput>
 {
-	static void FORCEINLINE DoKernel(VectorRegister* Dst, VectorRegister Src0)
+	static void VM_FORCEINLINE DoKernel(VectorRegister* Dst, VectorRegister Src0)
 	{
 		VectorStoreAlignedStreamed(Src0, Dst);
 	}
@@ -644,17 +764,33 @@ void VectorVM::Exec(
 			case EOp::sqrt: FVectorKernelSqrt::Exec(Context); break;
 			case EOp::neg: FVectorKernelNeg::Exec(Context); break;
 			case EOp::abs: FVectorKernelAbs::Exec(Context); break;
+			case EOp::exp: FVectorKernelExp::Exec(Context); break;
+			case EOp::exp2: FVectorKernelExp2::Exec(Context); break;
+			case EOp::log: FVectorKernelLog::Exec(Context); break;
+			case EOp::log2: FVectorKernelLog2::Exec(Context); break;
+			case EOp::sin: FVectorKernelSin::Exec(Context); break;
+			case EOp::cos: FVectorKernelCos::Exec(Context); break;
+			case EOp::tan: FVectorKernelTan::Exec(Context); break;
+			case EOp::asin: FVectorKernelASin::Exec(Context); break;
+			case EOp::acos: FVectorKernelACos::Exec(Context); break;
+			case EOp::atan: FVectorKernelATan::Exec(Context); break;
+			case EOp::atan2: FVectorKernelATan2::Exec(Context); break;
+			case EOp::ceil: FVectorKernelCeil::Exec(Context); break;
+			case EOp::floor: FVectorKernelFloor::Exec(Context); break;
+			case EOp::fmod: FVectorKernelMod::Exec(Context); break;
+			case EOp::frac: FVectorKernelFrac::Exec(Context); break;
+			case EOp::trunc: FVectorKernelTrunc::Exec(Context); break;
 			case EOp::clamp: FVectorKernelClamp::Exec(Context); break;
 			case EOp::min: FVectorKernelMin::Exec(Context); break;
 			case EOp::max: FVectorKernelMax::Exec(Context); break;
 			case EOp::pow: FVectorKernelPow::Exec(Context); break;
-			case EOp::sin: FVectorKernelSin::Exec(Context); break;
-			case EOp::sin4: FVectorKernelSin4::Exec(Context); break;
+			case EOp::sign: FVectorKernelSign::Exec(Context); break;
+			case EOp::step: FVectorKernelStep::Exec(Context); break;
 			case EOp::dot: FVectorKernelDot::Exec(Context); break;
-			case EOp::length: FVectorKernelLength::Exec(Context); break;
 			case EOp::cross: FVectorKernelCross::Exec(Context); break;
 			case EOp::normalize: FVectorKernelNormalize::Exec(Context); break;
 			case EOp::random: FVectorKernelRandom::Exec(Context); break;
+			case EOp::length: FVectorKernelLength::Exec(Context); break;
 			case EOp::noise: FVectorKernelNoise::Exec(Context); break;
 			case EOp::splatx: FVectorKernelSplat<0>::Exec(Context); break;
 			case EOp::splaty: FVectorKernelSplat<1>::Exec(Context); break;
@@ -783,3 +919,5 @@ bool FVectorVMTest::RunTest(const FString& Parameters)
 
 	return true;
 }
+
+#undef VM_FORCEINLINE
\ No newline at end of file
diff --git a/Engine/Source/Runtime/VectorVM/Public/VectorVM.h b/Engine/Source/Runtime/VectorVM/Public/VectorVM.h
--- a/Engine/Source/Runtime/VectorVM/Public/VectorVM.h
+++ b/Engine/Source/Runtime/VectorVM/Public/VectorVM.h
@@ -61,7 +61,6 @@ namespace VectorVM
 			normalize,
 			random,
 			length,
-			sin4,
 			noise,
 			splatx,
 			splaty,
