Niagara: various changes, fixes, and optimizations; first pass data objects (for buffer access, curves, etc.), to be enabled when niagara is in its own module (Curve data object sample currently always returns (1,0,0,1)).

diff --git a/Engine/Source/Editor/GraphEditor/Private/NodeFactory.cpp b/Engine/Source/Editor/GraphEditor/Private/NodeFactory.cpp
--- a/Engine/Source/Editor/GraphEditor/Private/NodeFactory.cpp
+++ b/Engine/Source/Editor/GraphEditor/Private/NodeFactory.cpp
@@ -391,6 +391,11 @@ TSharedPtr<SGraphPin> FNodeFactory::CreatePinWidget(UEdGraphPin* InPin)
 		{
 			return SNew(SGraphPin, InPin);
 		}
+		if (InPin->PinType.PinCategory == NSchema->PC_Curve)
+		{
+			return SNew(SGraphPin, InPin);
+		}
+
 	}
 
 
diff --git a/Engine/Source/Editor/NiagaraEditor/Private/NiagaraCompiler.cpp b/Engine/Source/Editor/NiagaraEditor/Private/NiagaraCompiler.cpp
--- a/Engine/Source/Editor/NiagaraEditor/Private/NiagaraCompiler.cpp
+++ b/Engine/Source/Editor/NiagaraEditor/Private/NiagaraCompiler.cpp
@@ -74,6 +74,10 @@ public:
 			SourceExpressions.Add(MakeShareable(new FNiagaraExpression(InCompiler, ENiagaraDataType::Vector)));
 			SourceExpressions.Add(MakeShareable(new FNiagaraExpression(InCompiler, ENiagaraDataType::Vector)));
 		}
+		else if (ResultType == ENiagaraDataType::Curve)
+		{
+			SourceExpressions.Add(MakeShareable(new FNiagaraExpression(InCompiler, ENiagaraDataType::Curve)));
+		}
 	}
 
 	virtual void Process()override
@@ -98,6 +102,10 @@ public:
 			Src3->ResultLocation = ENiagaraExpressionResultLocation::Constants;
 			Src3->ResultIndex = ResultIndex + 3;
 		}
+		else if (ResultType == ENiagaraDataType::Curve)
+		{
+			ResultLocation = ENiagaraExpressionResultLocation::BufferConstants;
+		}
 	}
 
 	FName ConstantName;
@@ -222,6 +230,11 @@ TNiagaraExprPtr FNiagaraCompiler::CompilePin(UEdGraphPin* Pin)
 				Schema->GetPinDefaultValue(Pin, Default);
 				ConstantData.SetOrAddInternal(ConstName, Default);
 			}
+			else if (DataType == ENiagaraDataType::Curve)
+			{
+				FNiagaraDataObject *Default = nullptr;
+				ConstantData.SetOrAddInternal(ConstName, Default);
+			}
 			Ret = Expression_GetInternalConstant(ConstName, DataType);
 		}
 	}
@@ -286,4 +299,10 @@ TNiagaraExprPtr FNiagaraCompiler::GetExternalConstant(FName ConstantName, const
 	return Expression_GetExternalConstant(ConstantName, ENiagaraDataType::Matrix);
 }
 
+TNiagaraExprPtr FNiagaraCompiler::GetExternalCurveConstant(FName ConstantName)
+{
+	SetOrAddConstant<FNiagaraDataObject*>(false, ConstantName, nullptr);
+	return Expression_GetExternalConstant(ConstantName, ENiagaraDataType::Curve);
+}
+
 #undef LOCTEXT_NAMESPACE
\ No newline at end of file
diff --git a/Engine/Source/Editor/NiagaraEditor/Private/NiagaraCompiler.h b/Engine/Source/Editor/NiagaraEditor/Private/NiagaraCompiler.h
--- a/Engine/Source/Editor/NiagaraEditor/Private/NiagaraCompiler.h
+++ b/Engine/Source/Editor/NiagaraEditor/Private/NiagaraCompiler.h
@@ -59,13 +59,14 @@ public:
 	virtual TNiagaraExprPtr GetExternalConstant(FName ConstantName, float Default)override;
 	virtual TNiagaraExprPtr GetExternalConstant(FName ConstantName, FVector4 Default)override;
 	virtual TNiagaraExprPtr GetExternalConstant(FName ConstantName, const FMatrix& Default)override;
+	virtual TNiagaraExprPtr GetExternalCurveConstant(FName ConstantName)override;
 
 	virtual void CheckInputs(FName OpName, TArray<TNiagaraExprPtr>& Inputs)override;
 	virtual void CheckOutputs(FName OpName, TArray<TNiagaraExprPtr>& Outputs)override;
 	//End INiagaraCompiler Interface
 	
 	/** Gets the index into a constants table of the constant specified by Name and bInternal. */
-	virtual void GetConstantResultIndex(FName Name, bool bInternal, int32& OutResultIndex, int32& OutComponentIndex) = 0;
+	virtual ENiagaraDataType GetConstantResultIndex(FName Name, bool bInternal, int32& OutResultIndex, int32& OutComponentIndex) = 0;
 	/**	Gets the index of a named attribute into either the input or output data. */
 	int32 GetAttributeIndex(FName Name);
 
diff --git a/Engine/Source/Editor/NiagaraEditor/Private/NiagaraCompiler_VectorVM.cpp b/Engine/Source/Editor/NiagaraEditor/Private/NiagaraCompiler_VectorVM.cpp
--- a/Engine/Source/Editor/NiagaraEditor/Private/NiagaraCompiler_VectorVM.cpp
+++ b/Engine/Source/Editor/NiagaraEditor/Private/NiagaraCompiler_VectorVM.cpp
@@ -70,13 +70,14 @@ public:
 		//Add the bitfield defining whether each source operand comes from a constant or a register.
 		int32 NumOperands = SourceExpressions.Num();
 		check(NumOperands <= 4);
-		VMCompiler->WriteCode(
-			VectorVM::CreateSrcOperandMask(
-			NumOperands > 0 ? VMCompiler->ExpressionIsConstant(GetSourceExpression(0).Get()) : false,
-			NumOperands > 1 ? VMCompiler->ExpressionIsConstant(GetSourceExpression(1).Get()) : false,
-			NumOperands > 2 ? VMCompiler->ExpressionIsConstant(GetSourceExpression(2).Get()) : false,
-			NumOperands > 3 ? VMCompiler->ExpressionIsConstant(GetSourceExpression(3).Get()) : false)
-			);
+
+
+		uint8 OpTypeMask = VectorVM::CreateSrcOperandMask(NumOperands > 0 ? VMCompiler->GetVMOperandType(GetSourceExpression(0).Get()) : VectorVM::RegisterOperandType,
+			NumOperands > 1 ? VMCompiler->GetVMOperandType(GetSourceExpression(1).Get()) : VectorVM::RegisterOperandType,
+			NumOperands > 2 ? VMCompiler->GetVMOperandType(GetSourceExpression(2).Get()) : VectorVM::RegisterOperandType,
+			NumOperands > 3 ? VMCompiler->GetVMOperandType(GetSourceExpression(3).Get()) : VectorVM::RegisterOperandType);
+		VMCompiler->WriteCode(OpTypeMask);
+
 
 		//Add the locations for each of the source operands.
 		for (int32 SrcIdx = 0; SrcIdx < SourceExpressions.Num(); ++SrcIdx)
@@ -116,7 +117,7 @@ public:
 
 		VMCompiler->WriteCode((uint8)VectorVM::EOp::output);
 		VMCompiler->WriteCode(DestIndex);
-		VMCompiler->WriteCode(VectorVM::CreateSrcOperandMask(VMCompiler->ExpressionIsConstant(SrcExpr)));
+		VMCompiler->WriteCode(VectorVM::CreateSrcOperandMask( VMCompiler->GetVMOperandType(SrcExpr) ) );
 		VMCompiler->WriteCode(SrcIndex);
 	}
 };
@@ -126,6 +127,18 @@ bool FNiagaraCompiler_VectorVM::ExpressionIsConstant(FNiagaraExpression*  Expres
 	return Expression->ResultLocation == ENiagaraExpressionResultLocation::Constants;
 }
 
+bool FNiagaraCompiler_VectorVM::ExpressionIsBufferConstant(FNiagaraExpression*  Expression)
+{
+	return Expression->ResultLocation == ENiagaraExpressionResultLocation::BufferConstants;
+}
+
+VectorVM::EOperandType FNiagaraCompiler_VectorVM::GetVMOperandType(FNiagaraExpression *Expression)
+{
+	return Expression->ResultLocation == ENiagaraExpressionResultLocation::Constants ? VectorVM::ConstantOperandType :
+		Expression->ResultLocation == ENiagaraExpressionResultLocation::BufferConstants ? VectorVM::DataObjConstantOperandType :
+		VectorVM::RegisterOperandType;
+}
+
 uint8 FNiagaraCompiler_VectorVM::GetResultVMIndex(FNiagaraExpression* Expression)
 {
 	if (Expression->ResultLocation == ENiagaraExpressionResultLocation::InputData)
@@ -138,14 +151,20 @@ uint8 FNiagaraCompiler_VectorVM::GetResultVMIndex(FNiagaraExpression* Expression
 	}
 	else
 	{
+		if (Expression->ResultLocation == ENiagaraExpressionResultLocation::BufferConstants)
+		{
+			return Expression->ResultIndex; //TempRegister or Constant.
+		}
+
 		return Expression->ResultIndex;//TempRegister or Constant.
 	}
 }
 
-void FNiagaraCompiler_VectorVM::GetConstantResultIndex(FName Name, bool bInternal, int32& OutResultIndex, int32& OutComponentIndex)
+ENiagaraDataType FNiagaraCompiler_VectorVM::GetConstantResultIndex(FName Name, bool bInternal, int32& OutResultIndex, int32& OutComponentIndex)
 {
 	ENiagaraDataType Type;
 	ConstantData.GetTableIndex(Name, bInternal, OutResultIndex, OutComponentIndex, Type);
+	return Type;
 }
 
 int32 FNiagaraCompiler_VectorVM::AquireTemporary()
@@ -594,4 +613,10 @@ void FNiagaraCompiler_VectorVM::LessThan_Internal(TArray<TNiagaraExprPtr>& Input
 	OutputExpressions.Add(Expression_VMNative(VectorVM::EOp::lessthan, InputExpressions));
 }
 
+
+void FNiagaraCompiler_VectorVM::Sample_Internal(TArray<TNiagaraExprPtr>& InputExpressions, TArray<TNiagaraExprPtr>& OutputExpressions)
+{
+	OutputExpressions.Add(Expression_VMNative(VectorVM::EOp::sample, InputExpressions));
+}
+
 #undef LOCTEXT_NAMESPACE
diff --git a/Engine/Source/Editor/NiagaraEditor/Private/NiagaraCompiler_VectorVM.h b/Engine/Source/Editor/NiagaraEditor/Private/NiagaraCompiler_VectorVM.h
--- a/Engine/Source/Editor/NiagaraEditor/Private/NiagaraCompiler_VectorVM.h
+++ b/Engine/Source/Editor/NiagaraEditor/Private/NiagaraCompiler_VectorVM.h
@@ -23,7 +23,9 @@ public:
 
 	/** Returns true if the expression is stored in the constant data. */
 	bool ExpressionIsConstant(FNiagaraExpression*  Expression);
-	
+	bool ExpressionIsBufferConstant(FNiagaraExpression*  Expression);
+	VectorVM::EOperandType GetVMOperandType(FNiagaraExpression *Expression);
+
 	/** Gets the final index used in the VM bytecode for the passed expression. */
 	uint8 GetResultVMIndex(FNiagaraExpression*  Expression);
 
@@ -41,7 +43,7 @@ public:
 	//End INiagaraCompiler;
 
 	//Begin FNiagaraCompiler Interface
-	virtual void GetConstantResultIndex(FName Name, bool bInternal, int32& OutResultIndex, int32& OutComponentIndex)override;
+	virtual ENiagaraDataType GetConstantResultIndex(FName Name, bool bInternal, int32& OutResultIndex, int32& OutComponentIndex)override;
 	virtual int32 AquireTemporary()override;
 	virtual void FreeTemporary(int32 TempIndex) override;
 	//End INiagaraCompiler Interface
diff --git a/Engine/Source/Editor/NiagaraEditor/Private/SNiagaraEffectEditorWidget.cpp b/Engine/Source/Editor/NiagaraEditor/Private/SNiagaraEffectEditorWidget.cpp
--- a/Engine/Source/Editor/NiagaraEditor/Private/SNiagaraEffectEditorWidget.cpp
+++ b/Engine/Source/Editor/NiagaraEditor/Private/SNiagaraEffectEditorWidget.cpp
@@ -170,16 +170,21 @@ void SEmitterWidget::Construct(const FArguments& InArgs)
 	RenderModuleList.Add(MakeShareable(new FString("Meshes")));
 
 	static TArray<TSharedPtr<EditorExposedVectorConstant>> DummyArray;
+	static TArray<TSharedPtr<EditorExposedVectorCurveConstant>> DummyCurveArray;
 	TArray<TSharedPtr<EditorExposedVectorConstant>> *EditorExposedConstants = &DummyArray;
+	TArray<TSharedPtr<EditorExposedVectorCurveConstant>> *EditorExposedCurveConstants = &DummyCurveArray;
 	if (Emitter->GetProperties()->UpdateScript && Emitter->GetProperties()->UpdateScript->Source)
 	{
 		EditorExposedConstants = &Emitter->GetProperties()->UpdateScript->Source->ExposedVectorConstants;
+		EditorExposedCurveConstants = &Emitter->GetProperties()->UpdateScript->Source->ExposedVectorCurveConstants;
 	}
 
 	TArray<TSharedPtr<EditorExposedVectorConstant>> *EditorExposedSpawnConstants = &DummyArray;
+	TArray<TSharedPtr<EditorExposedVectorCurveConstant>> *EditorExposedSpawnCurveConstants = &DummyCurveArray;
 	if (Emitter->GetProperties()->SpawnScript && Emitter->GetProperties()->SpawnScript->Source)
 	{
 		EditorExposedSpawnConstants = &Emitter->GetProperties()->SpawnScript->Source->ExposedVectorConstants;
+		EditorExposedSpawnCurveConstants = &Emitter->GetProperties()->SpawnScript->Source->ExposedVectorCurveConstants;
 	}
 
 
@@ -303,6 +308,20 @@ void SEmitterWidget::Construct(const FArguments& InArgs)
 											.SelectionMode(ESelectionMode::None)
 										]
 									]
+
+									+ SVerticalBox::Slot()
+										.Padding(2)
+										.AutoHeight()
+										.Expose(UpdateScriptCurveConstantListSlot)
+										[
+											SNew(SExpandableArea).InitiallyCollapsed(true).AreaTitle(FString("Curves"))
+											.BodyContent()
+											[
+												SNew(SListView<TSharedPtr<EditorExposedVectorCurveConstant>>)
+												.ItemHeight(20).ListItemsSource(EditorExposedCurveConstants).OnGenerateRow(this, &SEmitterWidget::OnGenerateCurveConstantListRow)
+												.SelectionMode(ESelectionMode::None)
+											]
+										]
 								]
 
 								+ SHorizontalBox::Slot()
diff --git a/Engine/Source/Editor/NiagaraEditor/Private/SNiagaraEffectEditorWidget.h b/Engine/Source/Editor/NiagaraEditor/Private/SNiagaraEffectEditorWidget.h
--- a/Engine/Source/Editor/NiagaraEditor/Private/SNiagaraEffectEditorWidget.h
+++ b/Engine/Source/Editor/NiagaraEditor/Private/SNiagaraEffectEditorWidget.h
@@ -17,6 +17,7 @@
 #include "ComponentReregisterContext.h"
 #include "SNumericEntryBox.h"
 
+
 #define NGED_SECTION_BORDER SNew(SBorder).BorderImage(FEditorStyle::GetBrush("ToolPanel.GroupBorder")).Padding(2.0f).HAlign(HAlign_Left)
 #define NGED_SECTION_LIGHTBORDER SNew(SBorder).BorderImage(FEditorStyle::GetBrush("ToolPanel.LightGroupBorder")).Padding(2.0f).HAlign(HAlign_Left)
 #define NGED_SECTION_DARKBORDER SNew(SBorder).BorderImage(FEditorStyle::GetBrush("ToolPanel.DarkGroupBorder")).Padding(2.0f).HAlign(HAlign_Left)
@@ -93,52 +94,24 @@ public:
 	{
 		FVector4 Constant(InVal, GetConstY().GetValue(), GetConstZ().GetValue(), GetConstW().GetValue());
 		Emitter->GetProperties()->ExternalConstants.SetOrAdd(ConstantName, Constant);
-		/*
-		FVector4 *VecPtr = Emitter->GetProperties()->ExternalConstants.FindVector(ConstantName);
-		if (VecPtr)
-		{
-			VecPtr->X = InVal;
-		}
-		*/
 	}
 
 	void OnConstantChangedY(float InVal, ETextCommit::Type Type)
 	{
 		FVector4 Constant(GetConstX().GetValue(), InVal, GetConstZ().GetValue(), GetConstW().GetValue());
 		Emitter->GetProperties()->ExternalConstants.SetOrAdd(ConstantName, Constant);
-		/*
-		FVector4 *VecPtr = Emitter->GetProperties()->ExternalConstants.FindVector(ConstantName);
-		if (VecPtr)
-		{
-			VecPtr->Y = InVal;
-		}
-		*/
 	}
 
 	void OnConstantChangedZ(float InVal, ETextCommit::Type Type)
 	{
 		FVector4 Constant(GetConstX().GetValue(), GetConstY().GetValue(), InVal, GetConstW().GetValue());
 		Emitter->GetProperties()->ExternalConstants.SetOrAdd(ConstantName, Constant);
-		/*
-		FVector4 *VecPtr = Emitter->GetProperties()->ExternalConstants.FindVector(ConstantName);
-		if (VecPtr)
-		{
-			VecPtr->Z = InVal;
-		}
-		*/
 	}
 
 	void OnConstantChangedW(float InVal, ETextCommit::Type Type)
 	{
 		FVector4 Constant(GetConstX().GetValue(), GetConstY().GetValue(), GetConstZ().GetValue(), InVal);
 		Emitter->GetProperties()->ExternalConstants.SetOrAdd(ConstantName, Constant);
-		/*
-		FVector4 *VecPtr = Emitter->GetProperties()->ExternalConstants.FindVector(ConstantName);
-		if (VecPtr)
-		{
-			VecPtr->W = InVal;
-		}
-		*/
 	}
 
 	void Construct(const FArguments& InArgs)
@@ -148,17 +121,18 @@ public:
 
 		FVector4 *ConstPtr = Emitter->GetProperties()->ExternalConstants.FindVector(ConstantName);
 		FVector4 ConstVal;
+
 		if (ConstPtr)
 		{
 			ConstVal = *ConstPtr;
 		}
-		
+
 		this->ChildSlot
 		[
 		SNew(SHorizontalBox)
 		+ SHorizontalBox::Slot().VAlign(VAlign_Center).Padding(4)[SNew(STextBlock).Text(FText::FromName(ConstantName))]
 		+ SHorizontalBox::Slot().VAlign(VAlign_Center).Padding(4)[SNew(SNumericEntryBox<float>).Value(this, &SVectorConstantWidget::GetConstX)
-		.OnValueCommitted(this, &SVectorConstantWidget::OnConstantChangedX).OnValueChanged(this, &SVectorConstantWidget::OnConstantChangedX)
+			.OnValueCommitted(this, &SVectorConstantWidget::OnConstantChangedX).OnValueChanged(this, &SVectorConstantWidget::OnConstantChangedX)
 			.Label()
 			[
 				SNumericEntryBox<float>::BuildLabel(FText::FromString("X"), FLinearColor::White, RedCol)
@@ -194,6 +168,74 @@ public:
 };
 
 
+
+
+class SCurveConstantWidget : public SCompoundWidget, public FNotifyHook
+{
+private:
+	FName ConstantName;
+	TSharedPtr<FNiagaraSimulation> Emitter;
+
+	UCurveVector *CurCurve;
+
+public:
+	SLATE_BEGIN_ARGS(SCurveConstantWidget) :
+		_Emitter(nullptr),
+		_ConstantName("Undefined")
+	{
+	}
+	SLATE_ARGUMENT(TSharedPtr<FNiagaraSimulation>, Emitter)
+		SLATE_ARGUMENT(FName, ConstantName)
+		SLATE_END_ARGS()
+
+	void OnCurveSelected(UObject *Asset)
+	{
+		CurCurve = Cast<UCurveVector>(Asset);
+		FNiagaraDataObject *DataObj = Emitter->GetProperties()->ExternalConstants.FindDataObj(ConstantName);
+		if (!DataObj)
+		{
+			DataObj = new FNiagaraCurveDataObject(nullptr);
+			Emitter->GetProperties()->ExternalConstants.SetOrAdd(ConstantName, DataObj);
+		}
+
+		FNiagaraCurveDataObject *CurveData = static_cast<FNiagaraCurveDataObject*>(DataObj);
+		CurveData->SetCurveObject(CurCurve);
+	}
+
+	UObject * GetCurve() const
+	{
+		FNiagaraDataObject *DataObj = Emitter->GetProperties()->ExternalConstants.FindDataObj(ConstantName);
+		FNiagaraCurveDataObject *CurveData = static_cast<FNiagaraCurveDataObject*>(DataObj);
+		if (CurveData)
+		{
+			return CurveData->GetCurveObject();
+		}
+		return nullptr;
+	}
+
+	void Construct(const FArguments& InArgs)
+	{
+		Emitter = InArgs._Emitter;
+		ConstantName = InArgs._ConstantName;
+
+		this->ChildSlot
+		[
+			SNew(SHorizontalBox)
+			+ SHorizontalBox::Slot().VAlign(VAlign_Center).Padding(4)[SNew(STextBlock).Text(FText::FromName(ConstantName))]
+			+ SHorizontalBox::Slot().VAlign(VAlign_Center).Padding(4)
+			[
+				SNew(SContentReference)
+				.AllowedClass(UCurveVector::StaticClass())
+				.AssetReference(this, &SCurveConstantWidget::GetCurve)
+				.AllowSelectingNewAsset(true)
+				.AllowClearingReference(true)
+				.OnSetReference(this, &SCurveConstantWidget::OnCurveSelected)
+				.WidthOverride(130)
+			]
+		];
+	}
+};
+
 class SEmitterWidget : public SCompoundWidget, public FNotifyHook
 {
 private:
@@ -213,7 +255,7 @@ private:
 	TArray<TSharedPtr<FString> > RenderModuleList;
 
 	TSharedPtr< SListView<TSharedPtr<EditorExposedVectorConstant>>> UpdateScriptConstantList;
-	SVerticalBox::FSlot *UpdateScriptConstantListSlot;
+	SVerticalBox::FSlot *UpdateScriptConstantListSlot, *UpdateScriptCurveConstantListSlot;
 	TSharedPtr< SListView<TSharedPtr<EditorExposedVectorConstant>>> SpawnScriptConstantList;
 	SVerticalBox::FSlot *SpawnScriptConstantListSlot;
 
@@ -236,6 +278,16 @@ public:
 	}
 
 
+	TSharedRef<ITableRow> OnGenerateCurveConstantListRow(TSharedPtr<EditorExposedVectorCurveConstant> InItem, const TSharedRef< STableViewBase >& OwnerTable)
+	{
+		return SNew(STableRow<TSharedPtr<FString>>, OwnerTable)
+			.Content()
+			[
+				SNew(SCurveConstantWidget).ConstantName(InItem->ConstName).Emitter(Emitter)
+			];
+	}
+
+
 	void OnEmitterNameChanged(const FText &InText)
 	{
 		Emitter->GetProperties()->Name = InText.ToString();
diff --git a/Engine/Source/Editor/UnrealEd/Classes/Niagara/EdGraphSchema_Niagara.h b/Engine/Source/Editor/UnrealEd/Classes/Niagara/EdGraphSchema_Niagara.h
--- a/Engine/Source/Editor/UnrealEd/Classes/Niagara/EdGraphSchema_Niagara.h
+++ b/Engine/Source/Editor/UnrealEd/Classes/Niagara/EdGraphSchema_Niagara.h
@@ -54,6 +54,8 @@ class UEdGraphSchema_Niagara : public UEdGraphSchema
 	FString PC_Vector;
 	UPROPERTY()
 	FString PC_Matrix;
+	UPROPERTY()
+	FString PC_Curve;
 
 	// Begin EdGraphSchema interface
 	virtual void GetGraphContextActions(FGraphContextMenuBuilder& ContextMenuBuilder) const override;
diff --git a/Engine/Source/Editor/UnrealEd/Private/Niagara/EdGraphSchema_Niagara.cpp b/Engine/Source/Editor/UnrealEd/Private/Niagara/EdGraphSchema_Niagara.cpp
--- a/Engine/Source/Editor/UnrealEd/Private/Niagara/EdGraphSchema_Niagara.cpp
+++ b/Engine/Source/Editor/UnrealEd/Private/Niagara/EdGraphSchema_Niagara.cpp
@@ -102,6 +102,7 @@ UEdGraphSchema_Niagara::UEdGraphSchema_Niagara(const FObjectInitializer& ObjectI
 	PC_Float = TEXT("float");
 	PC_Vector = TEXT("vector");
 	PC_Matrix = TEXT("matrix");
+	PC_Curve = TEXT("curve");
 }
 
 TSharedPtr<FNiagaraSchemaAction_NewNode> AddNewNodeAction(FGraphContextMenuBuilder& ContextMenuBuilder, const FString& Category, const FText& MenuDesc, const FString& Tooltip)
@@ -250,6 +251,10 @@ FLinearColor UEdGraphSchema_Niagara::GetPinTypeColor(const FEdGraphPinType& PinT
 	{
 		return Settings->TransformPinTypeColor;
 	}
+	else if (TypeString == PC_Curve)
+	{
+		return Settings->StructPinTypeColor;
+	}
 
 	return Settings->DefaultPinTypeColor;
 }
@@ -280,6 +285,10 @@ ENiagaraDataType UEdGraphSchema_Niagara::GetPinDataType(UEdGraphPin* Pin) const
 	{
 		return ENiagaraDataType::Matrix;
 	}
+	else if (Pin->PinType.PinCategory == PC_Curve)
+	{
+		return ENiagaraDataType::Curve;
+	}
 	return ENiagaraDataType::Invalid;
 }
 
diff --git a/Engine/Source/Editor/UnrealEd/Private/Niagara/NiagaraEditorCommon.cpp b/Engine/Source/Editor/UnrealEd/Private/Niagara/NiagaraEditorCommon.cpp
--- a/Engine/Source/Editor/UnrealEd/Private/Niagara/NiagaraEditorCommon.cpp
+++ b/Engine/Source/Editor/UnrealEd/Private/Niagara/NiagaraEditorCommon.cpp
@@ -519,6 +519,15 @@ void FNiagaraOpInfo::Init()
 	Op->Outputs.Add(FNiagaraOpInOutInfo(Result, ENiagaraDataType::Vector, ResultText, ResultText, DefaultStr_VecOne));
 	Op->OpDelegate.BindStatic(&INiagaraCompiler::LessThan);
 
+	Op = &OpInfos.Add(Sample);
+	Op->Name = Sample;
+	Op->FriendlyName = NSLOCTEXT("NiagaraOpInfo", "Sample Name", "Buffer Sample");
+	Op->Description = NSLOCTEXT("NiagaraOpInfo", "Sample Desc", "Result = Sample from buffer at index Time [0;1]");
+	Op->Inputs.Add(FNiagaraOpInOutInfo(A, ENiagaraDataType::Curve, NSLOCTEXT("NiagaraOpInfo", "Curve", "Curve"), NSLOCTEXT("NiagaraOpInfo", "Curve", "Curve"), DefaultStr_VecOne));
+	Op->Inputs.Add(FNiagaraOpInOutInfo(B, ENiagaraDataType::Vector, NSLOCTEXT("NiagaraOpInfo", "Time", "Time"), NSLOCTEXT("NiagaraOpInfo", "Time", "Time"), DefaultStr_VecOne));
+	Op->Outputs.Add(FNiagaraOpInOutInfo(Result, ENiagaraDataType::Vector, ResultText, ResultText, DefaultStr_VecOne));
+	Op->OpDelegate.BindStatic(&INiagaraCompiler::Sample);
+
 	#define NiagaraOp(OpName) if (GetOpInfo(OpName) == NULL)\
 	{\
 	UE_LOG(LogNiagaraCompiler, Fatal, TEXT("Couldn't find info for \"%s\". Have you added it to the NiagaraOpList but not provided the extra information in FNiagaraOpInfo::Init() ?"), TEXT(#OpName)); \
diff --git a/Engine/Source/Editor/UnrealEd/Private/Niagara/NiagaraNodeConstant.cpp b/Engine/Source/Editor/UnrealEd/Private/Niagara/NiagaraNodeConstant.cpp
--- a/Engine/Source/Editor/UnrealEd/Private/Niagara/NiagaraNodeConstant.cpp
+++ b/Engine/Source/Editor/UnrealEd/Private/Niagara/NiagaraNodeConstant.cpp
@@ -45,6 +45,11 @@ void UNiagaraNodeConstant::AllocateDefaultPins()
 		}
 		CreatePin(EGPD_Output, Schema->PC_Matrix, TEXT(""), NULL, false, false, ConstName.ToString());
 	}
+	else if (DataType == ENiagaraDataType::Curve)
+	{
+		//CreatePin(EGPD_Input, Schema->PC_Vector, TEXT(""), NULL, false, false, LOCTEXT("FloatPinName", "Time").ToString(), true);
+		CreatePin(EGPD_Output, Schema->PC_Curve, TEXT(""), NULL, false, false, ConstName.ToString());
+	}
 }
 
 FText UNiagaraNodeConstant::GetNodeTitle(ENodeTitleType::Type TitleType) const
@@ -129,6 +134,11 @@ void UNiagaraNodeConstant::Compile(class INiagaraCompiler* Compiler, TArray<FNia
 		Default.M[3][0] = Default3[0]; Default.M[3][1] = Default3[1]; Default.M[3][2] = Default3[2]; Default.M[3][3] = Default3[3];
 		Outputs.Add(FNiagaraNodeResult(Compiler->GetExternalConstant(ConstName, Default), OutPin));
 	}
+	else if (DataType == ENiagaraDataType::Curve)
+	{
+		UEdGraphPin* OutPin = Pins[0];
+		Outputs.Add(FNiagaraNodeResult(Compiler->GetExternalCurveConstant(ConstName), OutPin));
+	}
 }
 
 void UNiagaraNodeConstant::PostEditChangeProperty(struct FPropertyChangedEvent& PropertyChangedEvent)
diff --git a/Engine/Source/Editor/UnrealEd/Private/Niagara/NiagaraNodeOp.cpp b/Engine/Source/Editor/UnrealEd/Private/Niagara/NiagaraNodeOp.cpp
--- a/Engine/Source/Editor/UnrealEd/Private/Niagara/NiagaraNodeOp.cpp
+++ b/Engine/Source/Editor/UnrealEd/Private/Niagara/NiagaraNodeOp.cpp
@@ -52,6 +52,18 @@ void UNiagaraNodeOp::AllocateDefaultPins()
 			Pin->AutogeneratedDefaultValue = InOutInfo.Default;
 			Pin->PinToolTip = InOutInfo.Description.ToString();
 		}
+		else if (InOutInfo.DataType == ENiagaraDataType::Curve)
+		{
+			UEdGraphPin* Pin = CreatePin(EGPD_Input, Schema->PC_Curve, TEXT(""), NULL, false, false, InOutInfo.FriendlyName.ToString());
+			check(Pin);
+			Pin->bDefaultValueIsIgnored = true;
+			Pin->bDefaultValueIsReadOnly = false;
+			Pin->bNotConnectable = false;
+			Pin->DefaultValue = InOutInfo.Default;
+			Pin->AutogeneratedDefaultValue = InOutInfo.Default;
+			Pin->PinToolTip = InOutInfo.Description.ToString();
+		}
+
 	}
 	
 	for (int32 OutIdx = 0; OutIdx < OpInfo->Outputs.Num(); ++OutIdx)
@@ -75,6 +87,12 @@ void UNiagaraNodeOp::AllocateDefaultPins()
 			check(Pin);
 			Pin->PinToolTip = InOutInfo.Description.ToString();
 		}
+		else if (InOutInfo.DataType == ENiagaraDataType::Curve)
+		{
+			UEdGraphPin* Pin = CreatePin(EGPD_Output, Schema->PC_Curve, TEXT(""), NULL, false, false, InOutInfo.FriendlyName.ToString());
+			check(Pin);
+			Pin->PinToolTip = InOutInfo.Description.ToString();
+		}
 	}
 }
 
diff --git a/Engine/Source/Editor/UnrealEd/Private/Niagara/NiagaraScriptSource.cpp b/Engine/Source/Editor/UnrealEd/Private/Niagara/NiagaraScriptSource.cpp
--- a/Engine/Source/Editor/UnrealEd/Private/Niagara/NiagaraScriptSource.cpp
+++ b/Engine/Source/Editor/UnrealEd/Private/Niagara/NiagaraScriptSource.cpp
@@ -59,9 +59,18 @@ void UNiagaraScriptSource::Compile()
 	{
 		if (Node->bExposeToEffectEditor)
 		{
-			EditorExposedVectorConstant *Const = new EditorExposedVectorConstant();
-			Const->ConstName = Node->ConstName;
-			ExposedVectorConstants.Add( MakeShareable(Const) );
+			if(Node->DataType == ENiagaraDataType::Vector)
+			{
+				EditorExposedVectorConstant *Const = new EditorExposedVectorConstant();
+				Const->ConstName = Node->ConstName;
+				ExposedVectorConstants.Add( MakeShareable(Const) );
+			}
+			else if(Node->DataType == ENiagaraDataType::Curve)
+			{
+				EditorExposedVectorCurveConstant *Const = new EditorExposedVectorCurveConstant();
+				Const->ConstName = Node->ConstName;
+				ExposedVectorCurveConstants.Add(MakeShareable(Const));
+			}
 		}
 	}
 }
diff --git a/Engine/Source/Editor/UnrealEd/Public/INiagaraCompiler.h b/Engine/Source/Editor/UnrealEd/Public/INiagaraCompiler.h
--- a/Engine/Source/Editor/UnrealEd/Public/INiagaraCompiler.h
+++ b/Engine/Source/Editor/UnrealEd/Public/INiagaraCompiler.h
@@ -26,6 +26,9 @@ public:
 	virtual TNiagaraExprPtr GetExternalConstant(FName ConstantName, FVector4 Default) = 0;
 	/** Creates an expression which fetches the named constant. */
 	virtual TNiagaraExprPtr GetExternalConstant(FName ConstantName, const FMatrix& Default) = 0;
+	/** Creates an expression which fetches the named CURVE. */
+	virtual TNiagaraExprPtr GetExternalCurveConstant(FName ConstantName) = 0;
+
 	/** Creates an expression which writes the result of SourceExpression to the named attribute. */
 	virtual TNiagaraExprPtr Output(FName OutputName, TNiagaraExprPtr& SourceExpression) = 0;
 
diff --git a/Engine/Source/Editor/UnrealEd/Public/NiagaraEditorCommon.h b/Engine/Source/Editor/UnrealEd/Public/NiagaraEditorCommon.h
--- a/Engine/Source/Editor/UnrealEd/Public/NiagaraEditorCommon.h
+++ b/Engine/Source/Editor/UnrealEd/Public/NiagaraEditorCommon.h
@@ -61,7 +61,8 @@
 	NiagaraOp(TransformVector)\
 	NiagaraOp(Transpose)\
 	NiagaraOp(Inverse)\
-	NiagaraOp(LessThan)
+	NiagaraOp(LessThan)\
+	NiagaraOp(Sample)\
 
 enum class ENiagaraExpressionResultLocation
 {
@@ -69,6 +70,7 @@ enum class ENiagaraExpressionResultLocation
 	OutputData,
 	Constants,
 	Temporaries,
+	BufferConstants,
 	Unknown,
 };
 
diff --git a/Engine/Source/Runtime/Engine/Classes/Engine/NiagaraEffect.h b/Engine/Source/Runtime/Engine/Classes/Engine/NiagaraEffect.h
--- a/Engine/Source/Runtime/Engine/Classes/Engine/NiagaraEffect.h
+++ b/Engine/Source/Runtime/Engine/Classes/Engine/NiagaraEffect.h
@@ -123,6 +123,12 @@ public:
 		Constants.SetOrAdd(ConstantName, Value);
 	}
 
+	void SetConstant(FName ConstantName, FNiagaraDataObject *Value)
+	{
+		Constants.SetOrAdd(ConstantName, Value);
+	}
+
+
 
 	void Tick(float DeltaSeconds)
 	{
diff --git a/Engine/Source/Runtime/Engine/Classes/Engine/NiagaraScriptSourceBase.h b/Engine/Source/Runtime/Engine/Classes/Engine/NiagaraScriptSourceBase.h
--- a/Engine/Source/Runtime/Engine/Classes/Engine/NiagaraScriptSourceBase.h
+++ b/Engine/Source/Runtime/Engine/Classes/Engine/NiagaraScriptSourceBase.h
@@ -10,6 +10,12 @@ struct EditorExposedVectorConstant
 	FVector4 Value;
 };
 
+struct EditorExposedVectorCurveConstant
+{
+	FName ConstName;
+	class UCurveVector *Value;
+};
+
 /** Runtime data for a Niagara system */
 UCLASS(MinimalAPI)
 class UNiagaraScriptSourceBase : public UObject
@@ -17,5 +23,6 @@ class UNiagaraScriptSourceBase : public UObject
 	GENERATED_UCLASS_BODY()
 
 	TArray<TSharedPtr<EditorExposedVectorConstant> > ExposedVectorConstants;
+	TArray<TSharedPtr<EditorExposedVectorCurveConstant> > ExposedVectorCurveConstants;
 	virtual void Compile() {};
 };
\ No newline at end of file
diff --git a/Engine/Source/Runtime/Engine/Classes/Niagara/NiagaraCommon.h b/Engine/Source/Runtime/Engine/Classes/Niagara/NiagaraCommon.h
--- a/Engine/Source/Runtime/Engine/Classes/Niagara/NiagaraCommon.h
+++ b/Engine/Source/Runtime/Engine/Classes/Niagara/NiagaraCommon.h
@@ -10,6 +10,7 @@ enum class ENiagaraDataType : uint8
 	Scalar,
 	Vector,
 	Matrix,
+	Curve,
 	Unknown,
 	Invalid,
 };
diff --git a/Engine/Source/Runtime/Engine/Classes/Niagara/NiagaraConstantSet.h b/Engine/Source/Runtime/Engine/Classes/Niagara/NiagaraConstantSet.h
--- a/Engine/Source/Runtime/Engine/Classes/Niagara/NiagaraConstantSet.h
+++ b/Engine/Source/Runtime/Engine/Classes/Niagara/NiagaraConstantSet.h
@@ -15,6 +15,7 @@ private:
 	TMap<FName, float> ScalarConstants;
 	TMap<FName, FVector4> VectorConstants;
 	TMap<FName, FMatrix> MatrixConstants;
+	TMap<FName, class FNiagaraDataObject*> DataConstants;
 
 public:
 
@@ -33,12 +34,28 @@ public:
 		MatrixConstants.FindOrAdd(Name) = Mc;
 	}
 
+	void SetOrAdd(FName Name, class FNiagaraDataObject* Cc)
+	{
+		DataConstants.FindOrAdd(Name) = Cc;
+	}
+
+
 	float* FindScalar(FName InName){ return ScalarConstants.Find(InName); }
 	FVector4* FindVector(FName InName){ return VectorConstants.Find(InName); }
 	FMatrix* FindMatrix(FName InName){ return MatrixConstants.Find(InName); }
 	const float* FindScalar(FName InName)const { return ScalarConstants.Find(InName); }
 	const FVector4* FindVector(FName InName)const { return VectorConstants.Find(InName); }
 	const FMatrix* FindMatrix(FName InName)const { return MatrixConstants.Find(InName); }
+	FNiagaraDataObject* FindDataObj(FName InName) const 
+	{
+		FNiagaraDataObject * const *Ptr = DataConstants.Find(InName);
+		if (Ptr)
+		{
+			return *Ptr;
+		}
+
+		return nullptr; 
+	}
 
 	void Merge(FNiagaraConstantMap &InMap)
 	{
@@ -56,6 +73,11 @@ public:
 		{
 			SetOrAdd(MapIt.Key(), MapIt.Value());
 		}
+
+		for (auto MapIt = InMap.DataConstants.CreateIterator(); MapIt; ++MapIt)
+		{
+			SetOrAdd(MapIt.Key(), MapIt.Value());
+		}
 	}
 
 
@@ -91,6 +113,11 @@ private:
 	UPROPERTY()
 	TArray<FName> MatrixNames;
 
+	//UPROPERTY()
+	TArray<FNiagaraDataObject*> DataObjectConstants;
+	//UPROPERTY()
+	TArray<FName> DataObjectNames;
+
 public:
 
 	/** Fills the entire constants set into the constant table. */
@@ -121,6 +148,11 @@ public:
 		}
 	}
 
+	void AppendBufferConstants(TArray<class FNiagaraDataObject*> &Objs) const
+	{
+		Objs.Append(DataObjectConstants);
+	}
+
 	/** 
 	Fills only selected constants into the table. In the order they appear in the array of passed names not the order they appear in the set. 
 	Checks the passed map for entires to superceed the default values in the set.
@@ -160,6 +192,18 @@ public:
 		}
 	}
 
+	void AppendExternalBufferConstants(TArray<class FNiagaraDataObject*> &DataObjs, const FNiagaraConstantMap& Externals) const
+	{
+		for (int32 i = 0; i < DataObjectNames.Num(); ++i)
+		{
+			FName CcName = DataObjectNames[i];
+			FNiagaraDataObject* Data = Externals.FindDataObj(CcName);
+			DataObjs.Add(Data);
+		}
+	}
+
+
+
 	void SetOrAdd(FName Name, float Sc)
 	{
 		int32 Idx = ScalarNames.Find(Name);
@@ -203,6 +247,22 @@ public:
 		}
 	}
 
+
+
+	void SetOrAdd(FName Name, FNiagaraDataObject* Cc)
+	{
+		int32 Idx = DataObjectNames.Find(Name);
+		if (Idx == INDEX_NONE)
+		{
+			DataObjectNames.Add(Name);
+			DataObjectConstants.Add(Cc);
+		}
+		else
+		{
+			DataObjectConstants[Idx] = Cc;
+		}
+	}
+
 	FORCEINLINE int32 NumScalars()const { return ScalarNames.Num(); }
 	FORCEINLINE int32 NumVectors()const { return VectorNames.Num(); }
 	FORCEINLINE int32 NumMatrices()const { return MatrixNames.Num(); }
@@ -214,9 +274,40 @@ public:
 	FORCEINLINE int32 GetTableSize()const { return ScalarTableSize() + VectorTableSize() + MatrixTableSize(); }
 
 	/** Return the first constant used for this constant in a table. */
-	FORCEINLINE int32 GetTableIndex_Scalar(FName InName)const{ return ScalarNames.Find(InName) / 4; }
-	FORCEINLINE int32 GetTableIndex_Vector(FName InName)const{ return ScalarTableSize() + VectorNames.Find(InName); }
-	FORCEINLINE int32 GetTableIndex_Matrix(FName InName)const{ return ScalarTableSize() + VectorTableSize() + MatrixNames.Find(InName) * 4; }
+	FORCEINLINE int32 GetTableIndex_Scalar(FName InName)const
+	{ 
+		int Idx = ScalarNames.Find(InName);
+		if (Idx != INDEX_NONE)
+		{
+			return Idx / 4;
+		}
+		return INDEX_NONE;
+	}
+
+	FORCEINLINE int32 GetTableIndex_Vector(FName InName)const
+	{ 
+		int Idx = VectorNames.Find(InName);
+		if (Idx != INDEX_NONE)
+		{
+			return ScalarTableSize() + Idx;
+		}
+		return INDEX_NONE;
+	}
+
+	FORCEINLINE int32 GetTableIndex_Matrix(FName InName)const
+	{ 
+		int Idx = MatrixNames.Find(InName);
+		if (Idx != INDEX_NONE)
+		{
+			return ScalarTableSize() + VectorTableSize() + Idx * 4;
+		}
+		return INDEX_NONE;
+	}
+
+	FORCEINLINE int32 GetTableIndex_DataObj(FName InName)const
+	{ 
+		return DataObjectNames.Find(InName); 
+	}
 
 	/** Return the absolute index of the passed scalar. Can be used to get the component index also. */
 	FORCEINLINE int32 GetAbsoluteIndex_Scalar(FName InName)const{ return ScalarNames.Find(InName); }
diff --git a/Engine/Source/Runtime/Engine/Classes/Niagara/NiagaraScriptConstantData.h b/Engine/Source/Runtime/Engine/Classes/Niagara/NiagaraScriptConstantData.h
--- a/Engine/Source/Runtime/Engine/Classes/Niagara/NiagaraScriptConstantData.h
+++ b/Engine/Source/Runtime/Engine/Classes/Niagara/NiagaraScriptConstantData.h
@@ -20,14 +20,16 @@ struct FNiagaraScriptConstantData
 	FNiagaraConstants InternalConstants;
 
 	/** Fill a constants table ready for use in an update or spawn. */
-	void FillConstantTable(const FNiagaraConstantMap& ExternalConstantsMap, TArray<FVector4>& OutConstantTable)const
+	void FillConstantTable(const FNiagaraConstantMap& ExternalConstantsMap, TArray<FVector4>& OutConstantTable, TArray<FNiagaraDataObject *> &DataObjTable) const
 	{
 		//First up in the table comes the External constants in scalar->vector->matrix order.
 		//Only fills the constants actually used by the script.
 		OutConstantTable.Empty();
 		ExternalConstants.AppendToConstantsTable(OutConstantTable, ExternalConstantsMap);
+		ExternalConstants.AppendExternalBufferConstants(DataObjTable, ExternalConstantsMap);
 		//Next up add all the internal constants from the script.
 		InternalConstants.AppendToConstantsTable(OutConstantTable);
+		InternalConstants.AppendBufferConstants(DataObjTable);
 	}
 
 	void SetOrAddInternal(FName Name, float Sc)
@@ -45,6 +47,11 @@ struct FNiagaraScriptConstantData
 		InternalConstants.SetOrAdd(Name, Mc);
 	}
 
+	void SetOrAddInternal(FName Name, FNiagaraDataObject* Curve)
+	{
+		InternalConstants.SetOrAdd(Name, Curve);
+	}
+
 	void SetOrAddExternal(FName Name, float Sc)
 	{
 		ExternalConstants.SetOrAdd(Name, Sc);
@@ -60,6 +67,12 @@ struct FNiagaraScriptConstantData
 		ExternalConstants.SetOrAdd(Name, Mc);
 	}
 
+	void SetOrAddExternal(FName Name, FNiagaraDataObject* Curve)
+	{
+		ExternalConstants.SetOrAdd(Name, Curve);
+	}
+
+
 	/**
 	Calculates the index into a a constant table created from this constant data for the given constant name.
 	Obviously, assumes the constant data is complete. Any additions etc will invalidate previously calculated indexes.
@@ -82,6 +95,13 @@ struct FNiagaraScriptConstantData
 			{
 				ConstIdx = Consts.GetTableIndex_Matrix(Name);
 				Type = ENiagaraDataType::Matrix;
+
+				if (ConstIdx == INDEX_NONE)	// curves/buffers are in a separate table, so set base to 0
+				{
+					ConstIdx = Consts.GetTableIndex_DataObj(Name);
+					Type = ENiagaraDataType::Curve;
+					Base = 0;
+				}
 			}
 		}
 		else
diff --git a/Engine/Source/Runtime/Engine/Private/NiagaraSimulation.cpp b/Engine/Source/Runtime/Engine/Private/NiagaraSimulation.cpp
--- a/Engine/Source/Runtime/Engine/Private/NiagaraSimulation.cpp
+++ b/Engine/Source/Runtime/Engine/Private/NiagaraSimulation.cpp
@@ -125,7 +125,8 @@ void FNiagaraSimulation::UpdateParticles(
 
 	//Fill constant table with required emitter constants and internal script constants.
 	TArray<FVector4> ConstantTable;
-	Props->UpdateScript->ConstantData.FillConstantTable(Constants, ConstantTable);
+	TArray<FNiagaraDataObject *>DataObjTable;
+	Props->UpdateScript->ConstantData.FillConstantTable(Constants, ConstantTable, DataObjTable);
 
 	VectorVM::Exec(
 		Props->UpdateScript->ByteCode.GetData(),
@@ -134,6 +135,7 @@ void FNiagaraSimulation::UpdateParticles(
 		OutputRegisters,
 		NumAttr,
 		ConstantTable.GetData(),
+		DataObjTable.GetData(),
 		NumParticles
 		);
 }
@@ -171,7 +173,8 @@ int32 FNiagaraSimulation::SpawnAndKillParticles(int32 NumToSpawn)
 
 		//Fill constant table with required emitter constants and internal script constants.
 		TArray<FVector4> ConstantTable;
-		Props->SpawnScript->ConstantData.FillConstantTable(Constants, ConstantTable);
+		TArray<FNiagaraDataObject *>DataObjTable;
+		Props->SpawnScript->ConstantData.FillConstantTable(Constants, ConstantTable, DataObjTable);
 
 		VectorVM::Exec(
 			Props->SpawnScript->ByteCode.GetData(),
@@ -180,6 +183,7 @@ int32 FNiagaraSimulation::SpawnAndKillParticles(int32 NumToSpawn)
 			OutputRegisters,
 			NumAttr,
 			ConstantTable.GetData(),
+			DataObjTable.GetData(),
 			NumVectors
 			);
 	}
diff --git a/Engine/Source/Runtime/VectorVM/Private/VectorVM.cpp b/Engine/Source/Runtime/VectorVM/Private/VectorVM.cpp
--- a/Engine/Source/Runtime/VectorVM/Private/VectorVM.cpp
+++ b/Engine/Source/Runtime/VectorVM/Private/VectorVM.cpp
@@ -28,6 +28,16 @@ DEFINE_LOG_CATEGORY_STATIC(LogVectorVM, All, All);
 #define SRCOP_CCCR 0x0e
 #define SRCOP_CCCC 0x0f
 
+#define SRCOP_RRRB 0x10
+#define SRCOP_RRBR 0x20
+#define SRCOP_RRBB 0x30
+
+
+
+
+
+
+
 /**
  * Context information passed around during VM execution.
  */
@@ -39,24 +49,32 @@ struct FVectorVMContext
 	VectorRegister* RESTRICT * RESTRICT RegisterTable;
 	/** Pointer to the constant table. */
 	FVector4 const* RESTRICT ConstantTable;
+	/** Pointer to the data object constant table. */
+	FNiagaraDataObject * RESTRICT *DataObjConstantTable;
+
 	/** The number of vectors to process. */
 	int32 NumVectors;
 
 	/** Initialization constructor. */
 	FVectorVMContext(
-		uint8 const* InCode,
+		const uint8* InCode,
 		VectorRegister** InRegisterTable,
-		FVector4 const* InConstantTable,
+		const FVector4* InConstantTable,
+		FNiagaraDataObject** InDataObjTable,
 		int32 InNumVectors
 		)
 		: Code(InCode)
 		, RegisterTable(InRegisterTable)
 		, ConstantTable(InConstantTable)
+		, DataObjConstantTable(InDataObjTable)
 		, NumVectors(InNumVectors)
 	{
 	}
 };
 
+
+
+
 /** Decode the next operation contained in the bytecode. */
 static VM_FORCEINLINE VectorVM::EOp DecodeOp(FVectorVMContext& Context)
 {
@@ -76,6 +94,14 @@ static VM_FORCEINLINE VectorRegister DecodeConstant(FVectorVMContext& Context)
 	return VectorLoad(vec);
 }
 
+/** Decode a constant from the bytecode. */
+static /*VM_FORCEINLINE*/ const FNiagaraDataObject *DecodeDataObjConstant(FVectorVMContext& Context)
+{
+	const FNiagaraDataObject* Obj = Context.DataObjConstantTable[*Context.Code++];
+	return Obj;
+}
+
+
 static VM_FORCEINLINE uint8 DecodeSrcOperandTypes(FVectorVMContext& Context)
 {
 	return *Context.Code++;
@@ -86,6 +112,8 @@ static VM_FORCEINLINE uint8 DecodeByte(FVectorVMContext& Context)
 	return *Context.Code++;
 }
 
+
+
 /** Handles reading of a constant. */
 struct FConstantHandler
 {
@@ -96,6 +124,16 @@ struct FConstantHandler
 	VM_FORCEINLINE const VectorRegister& Get(){ return Constant; }
 };
 
+/** Handles reading of a data object constant. */
+struct FDataObjectConstantHandler
+{
+	const FNiagaraDataObject *Constant;
+	FDataObjectConstantHandler(FVectorVMContext& Context)
+		: Constant(DecodeDataObjConstant(Context))
+	{}
+	VM_FORCEINLINE const FNiagaraDataObject *Get(){ return Constant; }
+};
+
 /** Handles reading of a register, advancing the pointer with each read. */
 struct FRegisterHandler
 {
@@ -128,6 +166,7 @@ VM_FORCEINLINE void VectorBinaryLoop(FVectorVMContext& Context, VectorRegister*
 	}
 }
 
+
 template<typename Kernel, typename Arg0Handler, typename Arg1Handler, typename Arg2Handler>
 VM_FORCEINLINE void VectorTrinaryLoop(FVectorVMContext& Context, VectorRegister* RESTRICT Dst, int32 NumVectors)
 {
@@ -191,6 +230,29 @@ struct TBinaryVectorKernel
 	}
 };
 
+
+
+
+/** Base class of Vector kernels with 2 operands, one of which can be a data object. */
+template <typename Kernel>
+struct TBinaryVectorKernelData
+{
+	static void Exec(FVectorVMContext& Context)
+	{
+		VectorRegister* RESTRICT Dst = DecodeRegister(Context);
+		uint32 SrcOpTypes = DecodeSrcOperandTypes(Context);
+		int32 NumVectors = Context.NumVectors;
+		switch (SrcOpTypes)
+		{
+		case SRCOP_RRRB:	VectorBinaryLoop<Kernel, FDataObjectConstantHandler, FRegisterHandler>(Context, Dst, NumVectors); break;
+		default: check(0); break;
+		};
+	}
+};
+
+
+
+
 /** Base class of Vector kernels with 3 operands. */
 template <typename Kernel>
 struct TTrinaryVectorKernel
@@ -479,14 +541,11 @@ struct FVectorKernelTrunc : public TUnaryVectorKernel<FVectorKernelTrunc>
 
 struct FVectorKernelLessThan : public TBinaryVectorKernel<FVectorKernelLessThan>
 {
-	static void FORCEINLINE DoKernel(VectorRegister* RESTRICT Dst, VectorRegister Src0, VectorRegister Src1)
+	static void VM_FORCEINLINE DoKernel(VectorRegister* RESTRICT Dst, VectorRegister Src0, VectorRegister Src1)
 	{
 		const VectorRegister Large = MakeVectorRegister(BIG_NUMBER, BIG_NUMBER, BIG_NUMBER, BIG_NUMBER);
 		const VectorRegister One = MakeVectorRegister(1.0f, 1.0f, 1.0f, 1.0f);
 		const VectorRegister Zero = MakeVectorRegister(0.0f, 0.0f, 0.0f, 0.0f);
-
-		float const* FloatSrc0 = reinterpret_cast<float const*>(&Src0);
-		float const* FloatSrc1 = reinterpret_cast<float const*>(&Src1);
 		VectorRegister Tmp = VectorSubtract(Src1, Src0);
 		Tmp = VectorMultiply(Tmp, Large);
 		Tmp = VectorMin(Tmp, One);
@@ -494,6 +553,22 @@ struct FVectorKernelLessThan : public TBinaryVectorKernel<FVectorKernelLessThan>
 	}
 };
 
+
+
+struct FVectorKernelSample : public TBinaryVectorKernelData<FVectorKernelSample>
+{
+	static void FORCEINLINE DoKernel(VectorRegister* RESTRICT Dst, const FNiagaraDataObject *Src0, VectorRegister Src1)
+	{
+		if (Src0)
+		{
+			float const* FloatSrc1 = reinterpret_cast<float const*>(&Src1);
+			FVector4 Tmp = Src0->Sample(FloatSrc1[0]);
+			*Dst = VectorLoad(&Tmp);
+		}
+	}
+};
+
+
 struct FVectorKernelDot : public TBinaryVectorKernel<FVectorKernelDot>
 {
 	static void VM_FORCEINLINE DoKernel(VectorRegister* RESTRICT Dst, VectorRegister Src0, VectorRegister Src1)
@@ -506,10 +581,8 @@ struct FVectorKernelLength : public TUnaryVectorKernel<FVectorKernelLength>
 {
 	static void VM_FORCEINLINE DoKernel(VectorRegister* RESTRICT Dst, VectorRegister Src0)
 	{
-		VectorRegister Temp = VectorDot4(Src0, Src0);
-		float const* FloatSrc = reinterpret_cast<float const*>(&Temp);
-		float SDot = FMath::Sqrt(FloatSrc[0]);
-		*Dst = MakeVectorRegister(SDot, SDot, SDot, SDot);
+		VectorRegister Temp = VectorReciprocalLen(Src0);
+		*Dst = VectorReciprocal(Temp);
 	}
 };
 
@@ -591,51 +664,37 @@ struct FVectorKernelNoise : public TUnaryVectorKernel<FVectorKernelNoise>
 	static void VM_FORCEINLINE DoKernel(VectorRegister* RESTRICT Dst, VectorRegister Src0)
 	{
 		const VectorRegister One = MakeVectorRegister(1.0f, 1.0f, 1.0f, 1.0f);
+		const VectorRegister VecEight = MakeVectorRegister(8.0f, 8.0f, 8.0f, 8.0f);
 		const VectorRegister OneHalf = MakeVectorRegister(0.5f, 0.5f, 0.5f, 0.5f);
-		VectorRegister RndVals[2][2][2];
 		*Dst = MakeVectorRegister(0.0f, 0.0f, 0.0f, 0.0f);
-
-		float const* FloatSrc = reinterpret_cast<float const*>(&Src0);
-
+		
 		for (uint32 i = 1; i < 2; i++)
 		{
-			float FX = FloatSrc[0] / 5 / (1<<i);
-			float FY = FloatSrc[1] / 5 / (1<<i);
-			float FZ = FloatSrc[2] / 5 / (1<<i);
-			uint32 X = FMath::Abs( (int32)(FX) % 8 );
-			uint32 Y = FMath::Abs( (int32)(FY) % 8 );
-			uint32 Z = FMath::Abs( (int32)(FZ) % 8 );
-
-			for (int32 z = 0; z < 2; z++)
-			{
-				for (int32 y = 0; y < 2; y++)
-				{
-					for (int32 x = 0; x < 2; x++)
-					{
-						RndVals[x][y][z] = RandomTable[X+x][Y+y][Z+z];
-					}
-				}
-			}
-
-			const int32 IntCoords[3] = { static_cast<int32>(FX), static_cast<int32>(FY), static_cast<int32>(FZ) };
-			const float Fractionals[3] = { FX - IntCoords[0], FY - IntCoords[1], FZ - IntCoords[2] };
-			VectorRegister Alpha = MakeVectorRegister(Fractionals[0], Fractionals[0], Fractionals[0], Fractionals[0]);
-
+			float Di = 0.2f * (1.0f/(1<<i));
+			VectorRegister Div = MakeVectorRegister(Di, Di, Di, Di);
+			VectorRegister Coords = VectorMod( VectorAbs( VectorMultiply(Src0, Div) ), VecEight );
+			const float *CoordPtr = reinterpret_cast<float const*>(&Coords);
+			const int32 Cx = CoordPtr[0];
+			const int32 Cy = CoordPtr[1];
+			const int32 Cz = CoordPtr[2];
+
+			VectorRegister Frac = VectorFractional(Coords);
+			VectorRegister Alpha = VectorReplicate(Frac, 0);
 			VectorRegister OneMinusAlpha = VectorSubtract(One, Alpha);
-			VectorRegister XV1 = VectorAdd(VectorMultiply(RndVals[0][0][0], Alpha), VectorMultiply(RndVals[1][0][0], OneMinusAlpha));
-			VectorRegister XV2 = VectorAdd(VectorMultiply(RndVals[0][1][0], Alpha), VectorMultiply(RndVals[1][1][0], OneMinusAlpha));
-			VectorRegister XV3 = VectorAdd(VectorMultiply(RndVals[0][0][1], Alpha), VectorMultiply(RndVals[1][0][1], OneMinusAlpha));
-			VectorRegister XV4 = VectorAdd(VectorMultiply(RndVals[0][1][1], Alpha), VectorMultiply(RndVals[1][1][1], OneMinusAlpha));
+			
+			VectorRegister XV1 = VectorMultiplyAdd(RandomTable[Cx][Cy][Cz], Alpha, VectorMultiply(RandomTable[Cx+1][Cy][Cz], OneMinusAlpha));
+			VectorRegister XV2 = VectorMultiplyAdd(RandomTable[Cx][Cy+1][Cz], Alpha, VectorMultiply(RandomTable[Cx+1][Cy+1][Cz], OneMinusAlpha));
+			VectorRegister XV3 = VectorMultiplyAdd(RandomTable[Cx][Cy][Cz+1], Alpha, VectorMultiply(RandomTable[Cx+1][Cy][Cz+1], OneMinusAlpha));
+			VectorRegister XV4 = VectorMultiplyAdd(RandomTable[Cx][Cy+1][Cz+1], Alpha, VectorMultiply(RandomTable[Cx+1][Cy+1][Cz+1], OneMinusAlpha));
 
-			Alpha = MakeVectorRegister(Fractionals[1], Fractionals[1], Fractionals[1], Fractionals[1]);
+			Alpha = VectorReplicate(Frac, 1);
 			OneMinusAlpha = VectorSubtract(One, Alpha);
-			VectorRegister YV1 = VectorAdd(VectorMultiply(XV1, Alpha), VectorMultiply(XV2, OneMinusAlpha));
-			VectorRegister YV2 = VectorAdd(VectorMultiply(XV3, Alpha), VectorMultiply(XV4, OneMinusAlpha));
+			VectorRegister YV1 = VectorMultiplyAdd(XV1, Alpha, VectorMultiply(XV2, OneMinusAlpha));
+			VectorRegister YV2 = VectorMultiplyAdd(XV3, Alpha, VectorMultiply(XV4, OneMinusAlpha));
 
-
-			Alpha = MakeVectorRegister(Fractionals[2], Fractionals[2], Fractionals[2], Fractionals[2]);
+			Alpha = VectorReplicate(Frac, 2);
 			OneMinusAlpha = VectorSubtract(One, Alpha);
-			VectorRegister ZV = VectorAdd(VectorMultiply(YV1, Alpha), VectorMultiply(YV2, OneMinusAlpha));
+			VectorRegister ZV = VectorMultiplyAdd(YV1, Alpha, VectorMultiply(YV2, OneMinusAlpha));
 
 			*Dst = VectorAdd(*Dst, ZV);
 		}
@@ -708,6 +767,7 @@ void VectorVM::Exec(
 	VectorRegister** OutputRegisters,
 	int32 NumOutputRegisters,
 	FVector4 const* ConstantTable,
+	FNiagaraDataObject* *DataObjConstTable,
 	int32 NumVectors
 	)
 {
@@ -736,7 +796,7 @@ void VectorVM::Exec(
 
 		// Setup execution context.
 		int32 VectorsThisChunk = FMath::Min<int32>(NumVectors, VectorsPerChunk);
-		FVectorVMContext Context(Code, RegisterTable, ConstantTable, VectorsThisChunk);
+		FVectorVMContext Context(Code, RegisterTable, ConstantTable, DataObjConstTable, VectorsThisChunk);
 		EOp Op = EOp::done;
 
 		// Execute VM on all vectors in this chunk.
@@ -794,6 +854,7 @@ void VectorVM::Exec(
 			case EOp::composez: FVectorKernelCompose<2, 2, 2, 2>::Exec(Context); break;
 			case EOp::composew: FVectorKernelCompose<3, 3, 3, 3>::Exec(Context); break;
 			case EOp::lessthan: FVectorKernelLessThan::Exec(Context); break;
+			case EOp::sample: FVectorKernelSample::Exec(Context); break;
 			case EOp::output: FVectorKernelOutput::Exec(Context); break;
 
 			// Execution always terminates with a "done" opcode.
@@ -824,6 +885,18 @@ uint8 VectorVM::CreateSrcOperandMask(bool bIsOp0Constant, bool bIsOp1Constant, b
 			(bIsOp3Constant ? (1 << 3) : 0) ;
 }
 
+uint8 VectorVM::CreateSrcOperandMask(VectorVM::EOperandType Type1, VectorVM::EOperandType Type2, VectorVM::EOperandType Type3, VectorVM::EOperandType Type4)
+{
+	return	(Type1==VectorVM::ConstantOperandType ? (1 << 0) : 0) |
+		(Type2 == VectorVM::ConstantOperandType ? (1 << 1) : 0) |
+		(Type3 == VectorVM::ConstantOperandType ? (1 << 2) : 0) |
+		(Type4 == VectorVM::ConstantOperandType ? (1 << 3) : 0) |
+		(Type1 == VectorVM::DataObjConstantOperandType ? (1 << 4) : 0) | 
+		(Type2 == VectorVM::DataObjConstantOperandType ? (1 << 5) : 0) |
+		(Type3 == VectorVM::DataObjConstantOperandType ? (1 << 6) : 0) |
+		(Type4 == VectorVM::DataObjConstantOperandType ? (1 << 7) : 0);
+}
+
 /*------------------------------------------------------------------------------
 	Automation test for the VM.
 ------------------------------------------------------------------------------*/
@@ -869,6 +942,7 @@ bool FVectorVMTest::RunTest(const FString& Parameters)
 		InputRegisters, 3,
 		OutputRegisters, 1,
 		ConstantTable,
+		nullptr,
 		VectorVM::VectorsPerChunk
 		);
 
diff --git a/Engine/Source/Runtime/VectorVM/Public/VectorVM.h b/Engine/Source/Runtime/VectorVM/Public/VectorVM.h
--- a/Engine/Source/Runtime/VectorVM/Public/VectorVM.h
+++ b/Engine/Source/Runtime/VectorVM/Public/VectorVM.h
@@ -5,6 +5,40 @@
 #include "Core.h"
 
 
+/* Vector VM data object; encapsulates buffers, curves and other data in its derivatives
+*  for access by VectorVM kernels;
+*/
+class FNiagaraDataObject
+{
+public:
+	virtual FVector4 Sample(float InTime) const = 0;
+};
+
+/* Curve object; encapsulates a curve for the VectorVM
+*/
+class FNiagaraCurveDataObject : public FNiagaraDataObject
+{
+private:
+	class UCurveVector *CurveObj;
+public:
+	FNiagaraCurveDataObject(class UCurveVector *InCurve) : CurveObj(InCurve)
+	{
+	}
+
+	virtual FVector4 FNiagaraCurveDataObject::Sample(float InTime) const override
+	{
+		// commented out for the time being; until Niagara is in its own module, we can't depend on curve objects here
+		// because they're part of engine
+		//FVector Vec = CurveObj->GetVectorValue(InTime);	
+		//return FVector4(Vec, 0.0f);
+		return FVector4(1.0f, 0.0f, 0.0f, 1.0f);
+	}
+
+	UCurveVector *GetCurveObject()	{ return CurveObj; }
+	void SetCurveObject(UCurveVector *InCurve)	{ CurveObj = InCurve; }
+};
+
+
 namespace VectorVM
 {
 	/** Constants. */
@@ -20,6 +54,15 @@ namespace VectorVM
 		MaxRegisters = NumTempRegisters + MaxInputRegisters + MaxOutputRegisters + MaxConstants,
 	};
 
+	enum EOperandType
+	{
+		UndefinedOperandType = 0,
+		RegisterOperandType,
+		ConstantOperandType,
+		DataObjConstantOperandType,
+		TemporaryOperandType
+	};
+
 	/** List of opcodes supported by the VM. */
 	enum class EOp
 	{
@@ -73,6 +116,7 @@ namespace VectorVM
 			composew,
 			output,
 			lessthan,
+			sample,
 			NumOpcodes
 	};
 
@@ -80,6 +124,7 @@ namespace VectorVM
 	VECTORVM_API uint8 GetNumOpCodes();
 
 	VECTORVM_API uint8 CreateSrcOperandMask(bool bIsOp0Constant, bool bIsOp1Constant = false, bool bIsOp2Constant = false, bool bIsOp3Constant = false);
+	VECTORVM_API uint8 CreateSrcOperandMask(VectorVM::EOperandType Type1, VectorVM::EOperandType Type2 = VectorVM::RegisterOperandType, VectorVM::EOperandType Type3 = VectorVM::RegisterOperandType, VectorVM::EOperandType Type4 = VectorVM::RegisterOperandType);
 
 	/**
 	 * Execute VectorVM bytecode.
@@ -91,6 +136,7 @@ namespace VectorVM
 		VectorRegister** OutputRegisters,
 		int32 NumOutputRegisters,
 		FVector4 const* ConstantTable,
+		FNiagaraDataObject* *DataObjTable,
 		int32 NumVectors
 		);
 
@@ -98,3 +144,5 @@ namespace VectorVM
 
 
 } // namespace VectorVM
+
+
