Re-submitting compile fix.

diff --git a/Engine/Source/Runtime/VectorVM/Public/VectorVM.h b/Engine/Source/Runtime/VectorVM/Public/VectorVM.h
--- a/Engine/Source/Runtime/VectorVM/Public/VectorVM.h
+++ b/Engine/Source/Runtime/VectorVM/Public/VectorVM.h
@@ -25,7 +25,7 @@ public:
 	{
 	}
 
-	virtual FVector4 FNiagaraCurveDataObject::Sample(float InTime) const override
+	virtual FVector4 Sample(float InTime) const override
 	{
 		// commented out for the time being; until Niagara is in its own module, we can't depend on curve objects here
 		// because they're part of engine
