Fixing compile error. Stupid mistake, forgetting to completely remove an experiment.

diff --git a/Engine/Source/Runtime/VectorVM/Private/VectorVM.cpp b/Engine/Source/Runtime/VectorVM/Private/VectorVM.cpp
--- a/Engine/Source/Runtime/VectorVM/Private/VectorVM.cpp
+++ b/Engine/Source/Runtime/VectorVM/Private/VectorVM.cpp
@@ -2,7 +2,6 @@
 
 #include "VectorVMPrivate.h"
 #include "ModuleManager.h"
-#include <omp.h>
 
 IMPLEMENT_MODULE(FDefaultModuleImpl, VectorVM);
 
