Win32: Pass aligned variables by reference.

- Aligned types cannot be passed by values in Win32 (ABI restriction; appropriate stack alignment is not guaranteed).

diff --git a/Engine/Source/Runtime/VectorVM/Public/VectorVM.h b/Engine/Source/Runtime/VectorVM/Public/VectorVM.h
--- a/Engine/Source/Runtime/VectorVM/Public/VectorVM.h
+++ b/Engine/Source/Runtime/VectorVM/Public/VectorVM.h
@@ -11,8 +11,8 @@
 class FNiagaraDataObject
 {
 public:
-	virtual FVector4 Sample(FVector4 InCoords) const = 0;
-	virtual FVector4 Write(FVector4 InCoords, FVector4 Value) = 0;
+	virtual FVector4 Sample(const FVector4& InCoords) const = 0;
+	virtual FVector4 Write(const FVector4& InCoords, const FVector4 & Value) = 0;
 };
 
 /* Curve object; encapsulates a curve for the VectorVM
@@ -26,7 +26,7 @@ public:
 	{
 	}
 
-	virtual FVector4 Sample(FVector4 InCoords) const override
+	virtual FVector4 Sample(const FVector4& InCoords) const override
 	{
 		// commented out for the time being; until Niagara is in its own module, we can't depend on curve objects here
 		// because they're part of engine
@@ -35,7 +35,7 @@ public:
 		return FVector4(1.0f, 0.0f, 0.0f, 1.0f);
 	}
 
-	virtual FVector4 Write(FVector4 InCoords, FVector4 Value) override
+	virtual FVector4 Write(const FVector4& InCoords, const FVector4& Value) override
 	{
 		return FVector4(1.0f, 0.0f, 0.0f, 1.0f);
 	}
@@ -57,8 +57,9 @@ public:
 		Data.AddZeroed(32768);
 	}
 
-	uint32 MakeHash(FVector4 Coords) const
+	uint32 MakeHash(const FVector4& InCoords) const
 	{
+		FVector4 Coords = InCoords;
 		const uint32 P1 = 73856093; // some large primes
 		const uint32 P2 = 19349663;
 		const uint32 P3 = 83492791;
@@ -70,14 +71,14 @@ public:
 		return FMath::Clamp(N, 0, 32768);
 	}
 
-	virtual FVector4 Sample(FVector4 Coords) const override
+	virtual FVector4 Sample(const FVector4& Coords) const override
 	{
 		int32 Index = MakeHash(Coords);
 		Index = FMath::Clamp(Index, 0, Data.Num() - 1);
 		return Data[Index];
 	}
 
-	virtual FVector4 Write(FVector4 InCoords, FVector4 Value) override
+	virtual FVector4 Write(const FVector4& InCoords, const FVector4& Value) override
 	{
 		Data[MakeHash(InCoords)] = Value;
 		return Value;
