Missing file

diff --git a/Engine/Source/Runtime/VectorVM/Public/VectorVMDataObject.h b/Engine/Source/Runtime/VectorVM/Public/VectorVMDataObject.h
new file mode 100644
--- /dev/null
+++ b/Engine/Source/Runtime/VectorVM/Public/VectorVMDataObject.h
@@ -0,0 +1,99 @@
+// Copyright 1998-2015 Epic Games, Inc. All Rights Reserved.
+
+#pragma once	
+
+
+
+/* Vector VM data object; encapsulates buffers, curves and other data in its derivatives
+*  for access by VectorVM kernels;
+*/
+class FNiagaraDataObject
+{
+public:
+	virtual FVector4 Sample(FVector4 InCoords) const = 0;
+	virtual FVector4 Write(FVector4 InCoords, FVector4 Value) = 0;
+
+	virtual void Serialize(FArchive &Ar)
+	{
+
+	}
+	friend FArchive& operator<<(FArchive& Ar, FNiagaraDataObject* M)
+	{
+		M->Serialize(Ar);
+		return Ar;
+	}
+};
+
+/* Curve object; encapsulates a curve for the VectorVM
+*/
+class FNiagaraCurveDataObject : public FNiagaraDataObject
+{
+private:
+	class UCurveVector *CurveObj;
+public:
+	FNiagaraCurveDataObject(class UCurveVector *InCurve) : CurveObj(InCurve)
+	{
+	}
+
+	FVector4 Sample(FVector4 InCoords) const
+	{
+		FVector Vec = CurveObj->GetVectorValue(InCoords.X);
+		return FVector4(Vec, 0.0f);
+	}
+
+	virtual FVector4 Write(FVector4 InCoords, FVector4 Value) override
+	{
+		return FVector4(1.0f, 0.0f, 0.0f, 1.0f);
+	}
+
+	UCurveVector *GetCurveObject()	{ return CurveObj; }
+	void SetCurveObject(UCurveVector *InCurve)	{ CurveObj = InCurve; }
+
+
+	virtual void Serialize(FArchive &Ar) override
+	{
+		Ar << CurveObj;
+	}
+};
+
+
+/* Curve object; encapsulates sparse volumetric data for the VectorVM
+*/
+class FNiagaraSparseVolumeDataObject : public FNiagaraDataObject
+{
+private:
+	TArray<FVector4> Data;
+	uint64 NumBuckets;
+	int32 Size;
+public:
+	FNiagaraSparseVolumeDataObject()
+	{
+		Size = 64;
+		NumBuckets = Size*Size*Size;
+		Data.AddZeroed(NumBuckets);
+	}
+
+	uint32 MakeHash(FVector4 Coords) const
+	{
+		const int64 P1 = 73856093; // some large primes
+		const int64 P2 = 19349663;
+		const int64 P3 = 83492791;
+		Coords *= 0.1f;
+		uint64 N = (P1*static_cast<int64>(Coords.X)) ^ (P2*static_cast<int64>(Coords.Y)) ^ (P3*static_cast<int64>(Coords.Z));
+		N %= NumBuckets;
+		return FMath::Clamp(N, 0ULL, NumBuckets);
+	}
+
+	virtual FVector4 Sample(FVector4 Coords) const override
+	{
+		int32 Index = MakeHash(Coords);
+		Index = FMath::Clamp(Index, 0, Data.Num() - 1);
+		return Data[Index];
+	}
+
+	virtual FVector4 Write(FVector4 InCoords, FVector4 Value) override
+	{
+		Data[MakeHash(InCoords)] = Value;
+		return Value;
+	}
+};
\ No newline at end of file
