Fix build

diff --git a/Engine/Source/Runtime/VectorVM/Public/VectorVMDataObject.h b/Engine/Source/Runtime/VectorVM/Public/VectorVMDataObject.h
--- a/Engine/Source/Runtime/VectorVM/Public/VectorVMDataObject.h
+++ b/Engine/Source/Runtime/VectorVM/Public/VectorVMDataObject.h
@@ -10,8 +10,8 @@
 class FNiagaraDataObject
 {
 public:
-	virtual FVector4 Sample(FVector4 InCoords) const = 0;
-	virtual FVector4 Write(FVector4 InCoords, FVector4 Value) = 0;
+	virtual FVector4 Sample(const FVector4& InCoords) const = 0;
+	virtual FVector4 Write(const FVector4& InCoords, const FVector4& InValue) = 0;
 
 	virtual void Serialize(FArchive &Ar)
 	{
@@ -35,13 +35,13 @@ public:
 	{
 	}
 
-	FVector4 Sample(FVector4 InCoords) const
+	FVector4 Sample(const FVector4& InCoords) const
 	{
 		FVector Vec = CurveObj->GetVectorValue(InCoords.X);
 		return FVector4(Vec, 0.0f);
 	}
 
-	virtual FVector4 Write(FVector4 InCoords, FVector4 Value) override
+	virtual FVector4 Write(const FVector4& InCoords, const FVector4& InValue) override
 	{
 		return FVector4(1.0f, 0.0f, 0.0f, 1.0f);
 	}
@@ -73,8 +73,9 @@ public:
 		Data.AddZeroed(NumBuckets);
 	}
 
-	uint32 MakeHash(FVector4 Coords) const
+	uint32 MakeHash(const FVector4& InCoords) const
 	{
+		FVector4 Coords = InCoords;
 		const int64 P1 = 73856093; // some large primes
 		const int64 P2 = 19349663;
 		const int64 P3 = 83492791;
@@ -84,16 +85,16 @@ public:
 		return FMath::Clamp(N, 0ULL, NumBuckets);
 	}
 
-	virtual FVector4 Sample(FVector4 Coords) const override
+	virtual FVector4 Sample(const FVector4& InCoords) const override
 	{
-		int32 Index = MakeHash(Coords);
+		int32 Index = MakeHash(InCoords);
 		Index = FMath::Clamp(Index, 0, Data.Num() - 1);
 		return Data[Index];
 	}
 
-	virtual FVector4 Write(FVector4 InCoords, FVector4 Value) override
+	virtual FVector4 Write(const FVector4& InCoords, const FVector4& InValue) override
 	{
-		Data[MakeHash(InCoords)] = Value;
-		return Value;
+		Data[MakeHash(InCoords)] = InValue;
+		return InValue;
 	}
 };
\ No newline at end of file
