Niagara fixes:
- much improved noise, should be more or less divergence free now
- sin/cos/tan now have a full period within [0:1]
- added a Divide node

diff --git a/Engine/Source/Editor/NiagaraEditor/Private/NiagaraCompiler_VectorVM.cpp b/Engine/Source/Editor/NiagaraEditor/Private/NiagaraCompiler_VectorVM.cpp
--- a/Engine/Source/Editor/NiagaraEditor/Private/NiagaraCompiler_VectorVM.cpp
+++ b/Engine/Source/Editor/NiagaraEditor/Private/NiagaraCompiler_VectorVM.cpp
@@ -298,6 +298,11 @@ void FNiagaraCompiler_VectorVM::Multiply_Internal(TArray<TNiagaraExprPtr>& Input
 	OutputExpressions.Add(Expression_VMNative(VectorVM::EOp::mul, InputExpressions));
 }
 
+void FNiagaraCompiler_VectorVM::Divide_Internal(TArray<TNiagaraExprPtr>& InputExpressions, TArray<TNiagaraExprPtr>& OutputExpressions)
+{
+	OutputExpressions.Add(Expression_VMNative(VectorVM::EOp::div, InputExpressions));
+}
+
 void FNiagaraCompiler_VectorVM::MultiplyAdd_Internal(TArray<TNiagaraExprPtr>& InputExpressions, TArray<TNiagaraExprPtr>& OutputExpressions)
 {
 	OutputExpressions.Add(Expression_VMNative(VectorVM::EOp::mad, InputExpressions));
diff --git a/Engine/Source/Editor/UnrealEd/Private/Niagara/NiagaraEditorCommon.cpp b/Engine/Source/Editor/UnrealEd/Private/Niagara/NiagaraEditorCommon.cpp
--- a/Engine/Source/Editor/UnrealEd/Private/Niagara/NiagaraEditorCommon.cpp
+++ b/Engine/Source/Editor/UnrealEd/Private/Niagara/NiagaraEditorCommon.cpp
@@ -59,6 +59,15 @@ void FNiagaraOpInfo::Init()
 	Op->Outputs.Add(FNiagaraOpInOutInfo(Result, ENiagaraDataType::Vector, ResultText, ResultText, DefaultStr_VecOne));
 	Op->OpDelegate.BindStatic(&INiagaraCompiler::Multiply);
 
+	Op = &OpInfos.Add(Divide);
+	Op->Name = Divide;
+	Op->FriendlyName = NSLOCTEXT("NiagaraOpInfo", "Divide Name", "Divide");
+	Op->Description = NSLOCTEXT("NiagaraOpInfo", "Divide Desc", "Result = A / B = {{A.x / B.x, A.y / B.y, A.z / B.z, A.w / B.w}}");
+	Op->Inputs.Add(FNiagaraOpInOutInfo(A, ENiagaraDataType::Vector, AText, AText, DefaultStr_VecOne));
+	Op->Inputs.Add(FNiagaraOpInOutInfo(B, ENiagaraDataType::Vector, BText, BText, DefaultStr_VecOne));
+	Op->Outputs.Add(FNiagaraOpInOutInfo(Result, ENiagaraDataType::Vector, ResultText, ResultText, DefaultStr_VecOne));
+	Op->OpDelegate.BindStatic(&INiagaraCompiler::Divide);
+
 	Op = &OpInfos.Add(MultiplyAdd);
 	Op->Name = MultiplyAdd;
 	Op->FriendlyName = NSLOCTEXT("NiagaraOpInfo", "MultiplyAdd Name", "MultiplyAdd");
diff --git a/Engine/Source/Editor/UnrealEd/Public/NiagaraEditorCommon.h b/Engine/Source/Editor/UnrealEd/Public/NiagaraEditorCommon.h
--- a/Engine/Source/Editor/UnrealEd/Public/NiagaraEditorCommon.h
+++ b/Engine/Source/Editor/UnrealEd/Public/NiagaraEditorCommon.h
@@ -11,6 +11,7 @@
 	NiagaraOp(Add)\
 	NiagaraOp(Subtract)\
 	NiagaraOp(Multiply)\
+	NiagaraOp(Divide)\
 	NiagaraOp(MultiplyAdd)\
 	NiagaraOp(Lerp)\
 	NiagaraOp(Reciprocal)\
diff --git a/Engine/Source/Runtime/Engine/Private/NiagaraEffectRenderer.cpp b/Engine/Source/Runtime/Engine/Private/NiagaraEffectRenderer.cpp
--- a/Engine/Source/Runtime/Engine/Private/NiagaraEffectRenderer.cpp
+++ b/Engine/Source/Runtime/Engine/Private/NiagaraEffectRenderer.cpp
@@ -139,12 +139,12 @@ void NiagaraEffectRendererSprites::GetDynamicMeshElements(const TArray<const FSc
 				PerViewUniformParameters.RotationScale = 1.0f;
 				PerViewUniformParameters.RotationBias = 0.0f;
 				PerViewUniformParameters.TangentSelector = FVector4(0.0f, 0.0f, 0.0f, 1.0f);
-				PerViewUniformParameters.InvDeltaSeconds = 0.0f;
+				PerViewUniformParameters.InvDeltaSeconds = 30.0f;
 				if (Properties)
 				{
 					PerViewUniformParameters.SubImageSize = FVector4(Properties->SubImageInfo.X, Properties->SubImageInfo.Y, 1.0f / Properties->SubImageInfo.X, 1.0f / Properties->SubImageInfo.Y);
 				}
-				PerViewUniformParameters.NormalsType = 0;
+				PerViewUniformParameters.NormalsType = 0.0f;
 				PerViewUniformParameters.NormalsSphereCenter = FVector4(0.0f, 0.0f, 0.0f, 1.0f);
 				PerViewUniformParameters.NormalsCylinderUnitDirection = FVector4(0.0f, 0.0f, 1.0f, 0.0f);
 				PerViewUniformParameters.PivotOffset = FVector2D(-0.5f, -0.5f);
@@ -223,6 +223,7 @@ FNiagaraDynamicDataBase *NiagaraEffectRendererSprites::GenerateVertexData(const
 	CachedBounds.Init();
 
 	const FVector4 *PosPtr = Data.GetAttributeData(FNiagaraVariableInfo(FName(TEXT("Position")), ENiagaraDataType::Vector));
+	const FVector4 *VelPtr = Data.GetAttributeData(FNiagaraVariableInfo(FName(TEXT("Velocity")), ENiagaraDataType::Vector));
 	const FVector4 *ColPtr = Data.GetAttributeData(FNiagaraVariableInfo(FName(TEXT("Color")), ENiagaraDataType::Vector));
 	const FVector4 *AgePtr = Data.GetAttributeData(FNiagaraVariableInfo(FName(TEXT("Age")), ENiagaraDataType::Vector));
 	const FVector4 *RotPtr = Data.GetAttributeData(FNiagaraVariableInfo(FName(TEXT("Rotation")), ENiagaraDataType::Vector));
@@ -244,7 +245,10 @@ FNiagaraDynamicDataBase *NiagaraEffectRendererSprites::GenerateVertexData(const
 	{
 		FParticleSpriteVertex& NewVertex = RenderData[ParticleIndex];
 		NewVertex.Position = PosPtr[ParticleIndex];
-		NewVertex.OldPosition = NewVertex.Position;
+		if (VelPtr)
+		{
+			NewVertex.OldPosition = NewVertex.Position - VelPtr[ParticleIndex];
+		}
 		NewVertex.Color = FLinearColor(ColPtr[ParticleIndex]);
 		NewVertex.Color.A = ColPtr[ParticleIndex].W;
 		NewVertex.ParticleId = ParticleId;
@@ -390,8 +394,8 @@ void NiagaraEffectRendererRibbon::GetDynamicMeshElements(const TArray<const FSce
 
 				FNiagaraMeshCollectorResourcesRibbon& CollectorResources = Collector.AllocateOneFrameResource<FNiagaraMeshCollectorResourcesRibbon>();
 				FParticleBeamTrailUniformParameters PerViewUniformParameters;// = UniformParameters;
-				PerViewUniformParameters.CameraUp = FVector4(0.0f, 0.0f, 1.0f, 0.0f);
-				PerViewUniformParameters.CameraRight = FVector4(1.0f, 0.0f, 0.0f, 0.0f);
+				PerViewUniformParameters.CameraUp = View->GetViewUp(); // FVector4(0.0f, 0.0f, 1.0f, 0.0f);
+				PerViewUniformParameters.CameraRight = View->GetViewRight();//	FVector4(1.0f, 0.0f, 0.0f, 0.0f);
 				PerViewUniformParameters.ScreenAlignment = FVector4(0.0f, 0.0f, 0.0f, 0.0f);
 
 				// Collector.AllocateOneFrameResource uses default ctor, initialize the vertex factory
diff --git a/Engine/Source/Runtime/VectorVM/Private/VectorVM.cpp b/Engine/Source/Runtime/VectorVM/Private/VectorVM.cpp
--- a/Engine/Source/Runtime/VectorVM/Private/VectorVM.cpp
+++ b/Engine/Source/Runtime/VectorVM/Private/VectorVM.cpp
@@ -375,6 +375,15 @@ struct FVectorKernelMul : public TBinaryVectorKernel<FVectorKernelMul>
 	}
 };
 
+struct FVectorKernelDiv : public TBinaryVectorKernel<FVectorKernelDiv>
+{
+	static void VM_FORCEINLINE DoKernel(VectorRegister* RESTRICT Dst, VectorRegister Src0, VectorRegister Src1)
+	{
+		*Dst = VectorDivide(Src0, Src1);
+	}
+};
+
+
 struct FVectorKernelMad : public TTrinaryVectorKernel<FVectorKernelMad>
 {
 	static void VM_FORCEINLINE DoKernel(VectorRegister* RESTRICT Dst,VectorRegister Src0,VectorRegister Src1,VectorRegister Src2)
@@ -480,7 +489,7 @@ struct FVectorKernelSin : public TUnaryVectorKernel<FVectorKernelSin>
 {
 	static void VM_FORCEINLINE DoKernel(VectorRegister* RESTRICT Dst, VectorRegister Src0)
 	{
-		*Dst = VectorSin( VectorMultiply(Src0, GlobalVectorConstants::Pi));
+		*Dst = VectorSin( VectorMultiply(Src0, GlobalVectorConstants::PiByTwo));
 	}
 };
 
@@ -488,7 +497,7 @@ struct FVectorKernelCos : public TUnaryVectorKernel<FVectorKernelCos>
 {
 	static void VM_FORCEINLINE DoKernel(VectorRegister* RESTRICT Dst, VectorRegister Src0)
  	{
-		*Dst = VectorCos(VectorMultiply(Src0, GlobalVectorConstants::Pi));
+		*Dst = VectorCos(VectorMultiply(Src0, GlobalVectorConstants::PiByTwo));
 	}
 };
 
@@ -496,7 +505,7 @@ struct FVectorKernelTan : public TUnaryVectorKernel<FVectorKernelTan>
 {
 	static void VM_FORCEINLINE DoKernel(VectorRegister* RESTRICT Dst, VectorRegister Src0)
 	{
-		*Dst = VectorTan(VectorMultiply(Src0, GlobalVectorConstants::Pi));
+		*Dst = VectorTan(VectorMultiply(Src0, GlobalVectorConstants::PiByTwo));
 	}
 };
 
@@ -706,20 +715,20 @@ struct FVectorKernelStep : public TUnaryVectorKernel<FVectorKernelStep>
 
 struct FVectorKernelNoise : public TUnaryVectorKernel<FVectorKernelNoise>
 {
-	static VectorRegister RandomTable[10][10][10];
+	static VectorRegister RandomTable[17][17][17];
 
 	static void VM_FORCEINLINE DoKernel(VectorRegister* RESTRICT Dst, VectorRegister Src0)
 	{
 		const VectorRegister One = MakeVectorRegister(1.0f, 1.0f, 1.0f, 1.0f);
-		const VectorRegister VecEight = MakeVectorRegister(8.0f, 8.0f, 8.0f, 8.0f);
-		const VectorRegister OneHalf = MakeVectorRegister(0.5f, 0.5f, 0.5f, 0.5f);
+		const VectorRegister VecSize = MakeVectorRegister(16.0f, 16.0f, 16.0f, 16.0f);
+
 		*Dst = MakeVectorRegister(0.0f, 0.0f, 0.0f, 0.0f);
 		
 		for (uint32 i = 1; i < 2; i++)
 		{
 			float Di = 0.2f * (1.0f/(1<<i));
 			VectorRegister Div = MakeVectorRegister(Di, Di, Di, Di);
-			VectorRegister Coords = VectorMod( VectorAbs( VectorMultiply(Src0, Div) ), VecEight );
+			VectorRegister Coords = VectorMod( VectorAbs( VectorMultiply(Src0, Div) ), VecSize );
 			const float *CoordPtr = reinterpret_cast<float const*>(&Coords);
 			const int32 Cx = CoordPtr[0];
 			const int32 Cy = CoordPtr[1];
@@ -748,7 +757,7 @@ struct FVectorKernelNoise : public TUnaryVectorKernel<FVectorKernelNoise>
 	}
 };
 
-VectorRegister FVectorKernelNoise::RandomTable[10][10][10];
+VectorRegister FVectorKernelNoise::RandomTable[17][17][17];
 
 template<int32 Component>
 struct FVectorKernelSplat : public TUnaryVectorKernel<FVectorKernelSplat<Component>>
@@ -823,22 +832,82 @@ void VectorVM::Init()
 	static bool Inited = false;
 	if (Inited == false)
 	{
-		for (int z = 0; z < 10; z++)
+		// random noise
+		float TempTable[17][17][17];
+		for (int z = 0; z < 17; z++)
 		{
-			for (int y = 0; y < 10; y++)
+			for (int y = 0; y < 17; y++)
 			{
-				for (int x = 0; x < 10; x++)
+				for (int x = 0; x < 17; x++)
 				{
 					float f1 = (float)FMath::FRandRange(-1.0f, 1.0f);
-					float f2 = (float)FMath::FRandRange(-1.0f, 1.0f);
-					float f3 = (float)FMath::FRandRange(-1.0f, 1.0f);
-					float f4 = (float)FMath::FRandRange(-1.0f, 1.0f);
+					TempTable[x][y][z] = f1;
+				}
+			}
+		}
+
+		// pad
+		for (int i = 0; i < 17; i++)
+		{
+			for (int j = 0; j < 17; j++)
+			{
+				TempTable[i][j][16] = TempTable[i][j][0];
+				TempTable[i][16][j] = TempTable[i][0][j];
+				TempTable[16][j][i] = TempTable[0][j][i];
+			}
+		}
+
+		// compute gradients
+		FVector TempTable2[17][17][17];
+		for (int z = 0; z < 16; z++)
+		{
+			for (int y = 0; y < 16; y++)
+			{
+				for (int x = 0; x < 16; x++)
+				{
+					FVector XGrad = FVector(1.0f, 0.0f, TempTable[x][y][z] - TempTable[x+1][y][z]);
+					FVector YGrad = FVector(0.0f, 1.0f, TempTable[x][y][z] - TempTable[x][y + 1][z]);
+					FVector ZGrad = FVector(0.0f, 1.0f, TempTable[x][y][z] - TempTable[x][y][z+1]);
 
-					FVectorKernelNoise::RandomTable[x][y][z] = MakeVectorRegister(f1, f2, f3, f4);
+					FVector Grad = FVector(XGrad.Z, YGrad.Z, ZGrad.Z);
+					TempTable2[x][y][z] = Grad;
 				}
 			}
 		}
 
+		// pad
+		for (int i = 0; i < 17; i++)
+		{
+			for (int j = 0; j < 17; j++)
+			{
+				TempTable2[i][j][16] = TempTable2[i][j][0];
+				TempTable2[i][16][j] = TempTable2[i][0][j];
+				TempTable2[16][j][i] = TempTable2[0][j][i];
+			}
+		}
+
+
+		// compute curl of gradient field
+		for (int z = 0; z < 16; z++)
+		{
+			for (int y = 0; y < 16; y++)
+			{
+				for (int x = 0; x < 16; x++)
+				{
+					FVector Dy = TempTable2[x][y][z] - TempTable2[x][y + 1][z];
+					FVector Sy = TempTable2[x][y][z] + TempTable2[x][y + 1][z];
+					FVector Dx = TempTable2[x][y][z] - TempTable2[x + 1][y][z];
+					FVector Sx = TempTable2[x][y][z] + TempTable2[x + 1][y][z];
+					FVector Dz = TempTable2[x][y][z] - TempTable2[x][y][z + 1];
+					FVector Sz = TempTable2[x][y][z] + TempTable2[x][y][z + 1];
+					FVector Dir = FVector(Dy.Z - Sz.Y, Dz.X - Sx.Z, Dx.Y - Sy.X);
+
+					FVectorKernelNoise::RandomTable[x][y][z] = MakeVectorRegister(Dir.X, Dir.Y, Dir.Z, 0.0f);
+				}
+			}
+		}
+
+
 		Inited = true;
 	}
 }
@@ -892,6 +961,7 @@ void VectorVM::Exec(
 			case EOp::add: FVectorKernelAdd::Exec(Context); break;
 			case EOp::sub: FVectorKernelSub::Exec(Context); break;
 			case EOp::mul: FVectorKernelMul::Exec(Context); break;
+			case EOp::div: FVectorKernelDiv::Exec(Context); break;
 			case EOp::mad: FVectorKernelMad::Exec(Context); break;
 			case EOp::lerp: FVectorKernelLerp::Exec(Context); break;
 			case EOp::rcp: FVectorKernelRcp::Exec(Context); break;
diff --git a/Engine/Source/Runtime/VectorVM/Public/VectorVM.h b/Engine/Source/Runtime/VectorVM/Public/VectorVM.h
--- a/Engine/Source/Runtime/VectorVM/Public/VectorVM.h
+++ b/Engine/Source/Runtime/VectorVM/Public/VectorVM.h
@@ -88,6 +88,7 @@ namespace VectorVM
 			eventbroadcast,
 			easein,
 			easeinout,
+			div,
 			NumOpcodes
 	};
 
