Niagara: proper serialization of data objects; couple of bug fixes

diff --git a/Engine/Source/Editor/NiagaraEditor/Private/NiagaraCompiler.cpp b/Engine/Source/Editor/NiagaraEditor/Private/NiagaraCompiler.cpp
--- a/Engine/Source/Editor/NiagaraEditor/Private/NiagaraCompiler.cpp
+++ b/Engine/Source/Editor/NiagaraEditor/Private/NiagaraCompiler.cpp
@@ -263,7 +263,7 @@ TNiagaraExprPtr FNiagaraCompiler::CompilePin(UEdGraphPin* Pin)
 				break;
 			case ENiagaraDataType::Curve:
 			{
-				FNiagaraDataObject *Default = nullptr;
+				UNiagaraDataObject *Default = nullptr;
 				ConstantData.SetOrAddInternal(Var, Default);
 			}
 				break;
@@ -349,7 +349,7 @@ TNiagaraExprPtr FNiagaraCompiler::GetInternalConstant(const FNiagaraVariableInfo
 
 TNiagaraExprPtr FNiagaraCompiler::GetExternalCurveConstant(const FNiagaraVariableInfo& Constant)
 {
-	SetOrAddConstant<FNiagaraDataObject*>(false, Constant, nullptr);
+	SetOrAddConstant<UNiagaraDataObject*>(false, Constant, nullptr);
 	return Expression_GetExternalConstant(Constant);
 }
 
diff --git a/Engine/Source/Editor/NiagaraEditor/Private/NiagaraEffectEditor.cpp b/Engine/Source/Editor/NiagaraEditor/Private/NiagaraEffectEditor.cpp
--- a/Engine/Source/Editor/NiagaraEditor/Private/NiagaraEffectEditor.cpp
+++ b/Engine/Source/Editor/NiagaraEditor/Private/NiagaraEffectEditor.cpp
@@ -367,8 +367,8 @@ FReply FNiagaraEffectEditor::OnEmitterSelected(TSharedPtr<FNiagaraSimulation> Se
 	{
 		if (SelectedItem->GetProperties()->ExternalConstants.GetDataConstants().Num())
 		{
-			FNiagaraDataObject *DataObj = SelectedItem->GetProperties()->ExternalConstants.GetDataConstants().CreateConstIterator().Value();
-			TimeLine.Get()->SetCurve(static_cast<FNiagaraCurveDataObject*>(DataObj)->GetCurveObject());
+			UNiagaraDataObject *DataObj = SelectedItem->GetProperties()->ExternalConstants.GetDataConstants().CreateConstIterator().Value();
+			TimeLine.Get()->SetCurve(static_cast<UNiagaraCurveDataObject*>(DataObj)->GetCurveObject());
 		}
 	}
 	return FReply::Handled();
diff --git a/Engine/Source/Editor/NiagaraEditor/Private/SNiagaraEffectEditorWidget.h b/Engine/Source/Editor/NiagaraEditor/Private/SNiagaraEffectEditorWidget.h
--- a/Engine/Source/Editor/NiagaraEditor/Private/SNiagaraEffectEditorWidget.h
+++ b/Engine/Source/Editor/NiagaraEditor/Private/SNiagaraEffectEditorWidget.h
@@ -189,21 +189,21 @@ public:
 	void OnCurveSelected(UObject *Asset)
 	{
 		CurCurve = Cast<UCurveVector>(Asset);
-		FNiagaraDataObject *DataObj = Emitter->GetProperties()->ExternalConstants.FindDataObj(ConstantName);
+		UNiagaraDataObject *DataObj = Emitter->GetProperties()->ExternalConstants.FindDataObj(ConstantName);
 		if (!DataObj)
 		{
-			DataObj = new FNiagaraCurveDataObject(nullptr);
+			DataObj = NewObject<UNiagaraCurveDataObject>(Emitter->GetParentEffect());
 			Emitter->GetProperties()->ExternalConstants.SetOrAdd(ConstantName, DataObj);
 		}
 
-		FNiagaraCurveDataObject *CurveData = static_cast<FNiagaraCurveDataObject*>(DataObj);
+		UNiagaraCurveDataObject *CurveData = Cast<UNiagaraCurveDataObject>(DataObj);
 		CurveData->SetCurveObject(CurCurve);
 	}
 
 	UObject * GetCurve() const
 	{
-		FNiagaraDataObject *DataObj = Emitter->GetProperties()->ExternalConstants.FindDataObj(ConstantName);
-		FNiagaraCurveDataObject *CurveData = static_cast<FNiagaraCurveDataObject*>(DataObj);
+		UNiagaraDataObject *DataObj = Emitter->GetProperties()->ExternalConstants.FindDataObj(ConstantName);
+		UNiagaraCurveDataObject *CurveData = Cast<UNiagaraCurveDataObject>(DataObj);
 		if (CurveData)
 		{
 			return CurveData->GetCurveObject();
diff --git a/Engine/Source/Editor/UnrealEd/Private/Niagara/NiagaraScriptSource.cpp b/Engine/Source/Editor/UnrealEd/Private/Niagara/NiagaraScriptSource.cpp
--- a/Engine/Source/Editor/UnrealEd/Private/Niagara/NiagaraScriptSource.cpp
+++ b/Engine/Source/Editor/UnrealEd/Private/Niagara/NiagaraScriptSource.cpp
@@ -130,7 +130,7 @@ void UNiagaraScriptSource::Compile()
 	{
 		EditorExposedVectorCurveConstant *Const = new EditorExposedVectorCurveConstant();
 		FNiagaraVariableInfo Info;
-		FNiagaraDataObject* Dummy;
+		UNiagaraDataObject* Dummy;
 		ExternalConsts.GetDataObjectConstant(ConstIdx, Dummy, Info);
 		Const->ConstName = Info.Name;
 		//Set default value for this too?
diff --git a/Engine/Source/Runtime/Niagara/Classes/NiagaraEffect.h b/Engine/Source/Runtime/Niagara/Classes/NiagaraEffect.h
--- a/Engine/Source/Runtime/Niagara/Classes/NiagaraEffect.h
+++ b/Engine/Source/Runtime/Niagara/Classes/NiagaraEffect.h
@@ -59,20 +59,24 @@ class FNiagaraEffectInstance
 {
 public:
 	explicit FNiagaraEffectInstance(UNiagaraEffect *InAsset, UNiagaraComponent *InComponent)
+		: Effect(InAsset)
 	{
 		Component = InComponent;
 		for (int32 i = 0; i < InAsset->GetNumEmitters(); i++)
 		{
-			FNiagaraSimulation *Sim = new FNiagaraSimulation(InAsset->GetEmitterProperties(i), Component->GetWorld()->FeatureLevel);
+			FNiagaraSimulation *Sim = new FNiagaraSimulation(InAsset->GetEmitterProperties(i), InAsset, Component->GetWorld()->FeatureLevel);
 			Emitters.Add(MakeShareable(Sim));
 		}
 		InitRenderModules(Component->GetWorld()->FeatureLevel);
+		VolumeGrid = NewObject<UNiagaraSparseVolumeDataObject>();
 	}
 
 	explicit FNiagaraEffectInstance(UNiagaraEffect *InAsset)
 		: Component(nullptr)
+		, Effect(InAsset)
 	{
 		InitEmitters(InAsset);
+		VolumeGrid = NewObject<UNiagaraSparseVolumeDataObject>();
 	}
 
 	NIAGARA_API void InitEmitters(UNiagaraEffect *InAsset);
@@ -119,7 +123,7 @@ public:
 		Constants.SetOrAdd(ID, Value);
 	}
 
-	void SetConstant(FName ConstantName, FNiagaraDataObject *Value)
+	void SetConstant(FName ConstantName, UNiagaraDataObject *Value)
 	{
 		Constants.SetOrAdd(ConstantName, Value);
 	}
@@ -128,7 +132,7 @@ public:
 
 	void Tick(float DeltaSeconds)
 	{
-		Constants.SetOrAdd(FName(TEXT("EffectGrid")), &VolumeGrid); 
+		Constants.SetOrAdd(FName(TEXT("EffectGrid")), VolumeGrid); 
 
 		// pass the constants down to the emitter
 		// TODO: should probably just pass a pointer to the table
@@ -192,6 +196,7 @@ public:
 	FBox GetEffectBounds()	{ return EffectBounds;  }
 private:
 	UNiagaraComponent *Component;
+	UNiagaraEffect *Effect;
 	FBox EffectBounds;
 	float Age;
 
@@ -200,6 +205,5 @@ private:
 
 	TArray< TSharedPtr<FNiagaraSimulation> > Emitters;
 
-	FNiagaraSparseVolumeDataObject VolumeGrid;
-
+	UNiagaraSparseVolumeDataObject *VolumeGrid;
 };
\ No newline at end of file
diff --git a/Engine/Source/Runtime/Niagara/Classes/NiagaraSimulation.h b/Engine/Source/Runtime/Niagara/Classes/NiagaraSimulation.h
--- a/Engine/Source/Runtime/Niagara/Classes/NiagaraSimulation.h
+++ b/Engine/Source/Runtime/Niagara/Classes/NiagaraSimulation.h
@@ -95,8 +95,8 @@ public:
 struct FNiagaraSimulation
 {
 public:
-	explicit FNiagaraSimulation(FNiagaraEmitterProperties *InProps);
-	FNiagaraSimulation(FNiagaraEmitterProperties *Props, ERHIFeatureLevel::Type InFeatureLevel);
+	explicit FNiagaraSimulation(FNiagaraEmitterProperties *InProps, UNiagaraEffect *Effect);
+	FNiagaraSimulation(FNiagaraEmitterProperties *Props, UNiagaraEffect *Effect, ERHIFeatureLevel::Type InFeatureLevel);
 	virtual ~FNiagaraSimulation()
 	{}
 
@@ -124,6 +124,8 @@ public:
 	FNiagaraEmitterProperties *GetProperties()	{ return Props; }
 	void SetProperties(FNiagaraEmitterProperties *InProps);
 
+	UNiagaraEffect *GetParentEffect()	{ return ParentEffect; }
+
 	float NIAGARA_API GetTotalCPUTime();
 	int	NIAGARA_API GetTotalBytesUsed();
 
@@ -168,6 +170,7 @@ private:
 	FBox CachedBounds;
 
 	NiagaraEffectRenderer *EffectRenderer;
+	UNiagaraEffect *ParentEffect;
 
 	/** Calc number to spawn */
 	int32 CalcNumToSpawn(float DeltaSeconds)
diff --git a/Engine/Source/Runtime/Niagara/NiagaraConstantSet.h b/Engine/Source/Runtime/Niagara/NiagaraConstantSet.h
--- a/Engine/Source/Runtime/Niagara/NiagaraConstantSet.h
+++ b/Engine/Source/Runtime/Niagara/NiagaraConstantSet.h
@@ -3,9 +3,9 @@
 #pragma once
 
 #include "NiagaraCommon.h"
+#include "VectorVMDataObject.h"
 #include "NiagaraConstantSet.generated.h"
 
-
 USTRUCT()
 struct FNiagaraConstantMap
 {
@@ -15,7 +15,7 @@ private:
 	TMap<FNiagaraVariableInfo, float> ScalarConstants;
 	TMap<FNiagaraVariableInfo, FVector4> VectorConstants;
 	TMap<FNiagaraVariableInfo, FMatrix> MatrixConstants;
-	TMap<FNiagaraVariableInfo, class FNiagaraDataObject*> DataConstants;
+	TMap<FNiagaraVariableInfo, UNiagaraDataObject*> DataConstants;
 
 public:
 
@@ -44,7 +44,7 @@ public:
 		MatrixConstants.FindOrAdd(ID) = Mc;
 	}
 
-	void SetOrAdd(FNiagaraVariableInfo ID, FNiagaraDataObject* Cc)
+	void SetOrAdd(FNiagaraVariableInfo ID, UNiagaraDataObject* Cc)
 	{
 		check(ID.Type == ENiagaraDataType::Curve);
 		DataConstants.FindOrAdd(ID) = Cc;
@@ -65,7 +65,7 @@ public:
 		MatrixConstants.FindOrAdd(FNiagaraVariableInfo(Name, ENiagaraDataType::Matrix)) = Mc;
 	}
 
-	void SetOrAdd(FName Name, class FNiagaraDataObject* Cc)
+	void SetOrAdd(FName Name, UNiagaraDataObject* Cc)
 	{
 		DataConstants.FindOrAdd(FNiagaraVariableInfo(Name, ENiagaraDataType::Curve)) = Cc;
 	}
@@ -77,9 +77,9 @@ public:
 	const float* FindScalar(FNiagaraVariableInfo ID)const { return ScalarConstants.Find(ID); }
 	const FVector4* FindVector(FNiagaraVariableInfo ID)const { return VectorConstants.Find(ID); }
 	const FMatrix* FindMatrix(FNiagaraVariableInfo ID)const { return MatrixConstants.Find(ID); }
-	FNiagaraDataObject* FindDataObj(FNiagaraVariableInfo ID) const
+	UNiagaraDataObject* FindDataObj(FNiagaraVariableInfo ID) const
 	{
-		FNiagaraDataObject * const *Ptr = DataConstants.Find(ID);
+		UNiagaraDataObject * const *Ptr = DataConstants.Find(ID);
 		if (Ptr)
 		{
 			return *Ptr;
@@ -90,7 +90,7 @@ public:
 	float* FindScalar(FName Name) { return ScalarConstants.Find(FNiagaraVariableInfo(Name,ENiagaraDataType::Scalar)); }
 	FVector4* FindVector(FName Name) { return VectorConstants.Find(FNiagaraVariableInfo(Name, ENiagaraDataType::Vector)); }
 	FMatrix* FindMatrix(FName Name) { return MatrixConstants.Find(FNiagaraVariableInfo(Name, ENiagaraDataType::Matrix)); }
-	FNiagaraDataObject* FindDataObj(FName Name)const { return FindDataObj(FNiagaraVariableInfo(Name, ENiagaraDataType::Curve)); }
+	UNiagaraDataObject* FindDataObj(FName Name)const { return FindDataObj(FNiagaraVariableInfo(Name, ENiagaraDataType::Curve)); }
 
 	void Merge(FNiagaraConstantMap &InMap)
 	{
@@ -115,7 +115,7 @@ public:
 		}
 	}
 
-	const TMap<FNiagaraVariableInfo, class FNiagaraDataObject*> &GetDataConstants()	const
+	const TMap<FNiagaraVariableInfo, UNiagaraDataObject*> &GetDataConstants()	const
 	{ 
 		return DataConstants; 
 	}
@@ -124,7 +124,7 @@ public:
 	virtual bool Serialize(FArchive &Ar)
 	{
 		// TODO: can't serialize the data object constants at the moment; need to figure that out
-		Ar << ScalarConstants << VectorConstants << MatrixConstants;
+		Ar << ScalarConstants << VectorConstants << MatrixConstants << DataConstants;
 		return true;
 	}
 };
@@ -155,7 +155,7 @@ private:
 	UPROPERTY()
 	TArray<FMatrix> MatrixConstants;
 	//UPROPERTY()
-	TArray<FNiagaraDataObject*> DataObjectConstants;
+	TArray<UNiagaraDataObject*> DataObjectConstants;
 
 public:
 
@@ -186,7 +186,7 @@ public:
 		OutInfo = MatrixConstantsInfo[Index];
 		OutValue = MatrixConstants[Index];
 	}
-	void GetDataObjectConstant(int32 Index, FNiagaraDataObject*& OutValue, FNiagaraVariableInfo& OutInfo)const
+	void GetDataObjectConstant(int32 Index, UNiagaraDataObject*& OutValue, FNiagaraVariableInfo& OutInfo)const
 	{
 		OutInfo = DataObjectConstantsInfo[Index];
 		OutValue = DataObjectConstants[Index];
@@ -225,7 +225,7 @@ public:
 		}
 	}
 
-	void AppendBufferConstants(TArray<class FNiagaraDataObject*> &Objs) const
+	void AppendBufferConstants(TArray<class UNiagaraDataObject*> &Objs) const
 	{
 		Objs.Append(DataObjectConstants);
 	}
@@ -266,12 +266,12 @@ public:
 		}
 	}
 
-	void AppendExternalBufferConstants(TArray<class FNiagaraDataObject*> &DataObjs, const FNiagaraConstantMap& Externals) const
+	void AppendExternalBufferConstants(TArray<class UNiagaraDataObject*> &DataObjs, const FNiagaraConstantMap& Externals) const
 	{
 		for (int32 i = 0; i < DataObjectConstantsInfo.Num(); ++i)
 		{
 			FNiagaraVariableInfo CcID = DataObjectConstantsInfo[i];
-			FNiagaraDataObject* Data = Externals.FindDataObj(CcID);
+			UNiagaraDataObject* Data = Externals.FindDataObj(CcID);
 			DataObjs.Add(Data);
 		}
 	}
@@ -318,7 +318,7 @@ public:
 		}
 	}
 
-	void SetOrAdd(const FNiagaraVariableInfo& Constant, FNiagaraDataObject* Value)
+	void SetOrAdd(const FNiagaraVariableInfo& Constant, UNiagaraDataObject* Value)
 	{
 		int32 Idx = DataObjectConstantsInfo.Find(Constant);
 		if (Idx == INDEX_NONE)
diff --git a/Engine/Source/Runtime/Niagara/NiagaraScriptConstantData.h b/Engine/Source/Runtime/Niagara/NiagaraScriptConstantData.h
--- a/Engine/Source/Runtime/Niagara/NiagaraScriptConstantData.h
+++ b/Engine/Source/Runtime/Niagara/NiagaraScriptConstantData.h
@@ -26,7 +26,7 @@ struct FNiagaraScriptConstantData
 	}
 
 	/** Fill a constants table ready for use in an update or spawn. */
-	void FillConstantTable(const FNiagaraConstantMap& ExternalConstantsMap, TArray<FVector4>& OutConstantTable, TArray<FNiagaraDataObject *> &DataObjTable) const
+	void FillConstantTable(const FNiagaraConstantMap& ExternalConstantsMap, TArray<FVector4>& OutConstantTable, TArray<UNiagaraDataObject *> &DataObjTable) const
 	{
 		//First up in the table comes the External constants in scalar->vector->matrix order.
 		//Only fills the constants actually used by the script.
diff --git a/Engine/Source/Runtime/Niagara/Private/NiagaraEffect.cpp b/Engine/Source/Runtime/Niagara/Private/NiagaraEffect.cpp
--- a/Engine/Source/Runtime/Niagara/Private/NiagaraEffect.cpp
+++ b/Engine/Source/Runtime/Niagara/Private/NiagaraEffect.cpp
@@ -62,7 +62,7 @@ void UNiagaraEffect::PreSave()
 
 TSharedPtr<FNiagaraSimulation> FNiagaraEffectInstance::AddEmitter(FNiagaraEmitterProperties *Properties)
 {
-	FNiagaraSimulation *SimPtr = new FNiagaraSimulation(Properties);
+	FNiagaraSimulation *SimPtr = new FNiagaraSimulation(Properties, Effect);
 	TSharedPtr<FNiagaraSimulation> Sim = MakeShareable(SimPtr);
 
 	Sim->SetRenderModuleType(RMT_Sprites, Component->GetWorld()->FeatureLevel);
@@ -90,7 +90,7 @@ void FNiagaraEffectInstance::InitEmitters(UNiagaraEffect *InAsset)
 	for (int i = 0; i < InAsset->GetNumEmitters(); i++)
 	{
 		FNiagaraEmitterProperties *Props = InAsset->GetEmitterProperties(i);
-		FNiagaraSimulation *Sim = new FNiagaraSimulation(Props);
+		FNiagaraSimulation *Sim = new FNiagaraSimulation(Props, InAsset);
 		Emitters.Add(MakeShareable(Sim));
 	}
 }
diff --git a/Engine/Source/Runtime/Niagara/Private/NiagaraSimulation.cpp b/Engine/Source/Runtime/Niagara/Private/NiagaraSimulation.cpp
--- a/Engine/Source/Runtime/Niagara/Private/NiagaraSimulation.cpp
+++ b/Engine/Source/Runtime/Niagara/Private/NiagaraSimulation.cpp
@@ -5,26 +5,28 @@
 #include "NiagaraEffectRenderer.h"
 #include "VectorVM.h"
 
-FNiagaraSimulation::FNiagaraSimulation(FNiagaraEmitterProperties *InProps) 
+FNiagaraSimulation::FNiagaraSimulation(FNiagaraEmitterProperties *InProps, UNiagaraEffect *InEffect)
 : Age(0.0f)
 , Loops(0)
 , bIsEnabled(true)
 , SpawnRemainder(0.0f)
 , CachedBounds(ForceInit)
 , EffectRenderer(nullptr)
+, ParentEffect(InEffect)
 {
 	Props = InProps;
 
 	Init();
 }
 
-FNiagaraSimulation::FNiagaraSimulation(FNiagaraEmitterProperties *InProps, ERHIFeatureLevel::Type InFeatureLevel)
+FNiagaraSimulation::FNiagaraSimulation(FNiagaraEmitterProperties *InProps, UNiagaraEffect *InEffect, ERHIFeatureLevel::Type InFeatureLevel)
 	: Age(0.0f)
 	, Loops(0)
 	, bIsEnabled(true)
 	, SpawnRemainder(0.0f)
 	, CachedBounds(ForceInit)
 	, EffectRenderer(nullptr)
+	, ParentEffect(InEffect)
 {
 	Props = InProps;
 
@@ -274,7 +276,7 @@ void FNiagaraSimulation::RunVMScript(UNiagaraScript* Script, EUnusedAttributeBeh
 
 	//Fill constant table with required emitter constants and internal script constants.
 	TArray<FVector4> ConstantTable;
-	TArray<FNiagaraDataObject *>DataObjTable;
+	TArray<UNiagaraDataObject *>DataObjTable;
 	Script->ConstantData.FillConstantTable(Constants, ConstantTable, DataObjTable);
 
 	VectorVM::Exec(
diff --git a/Engine/Source/Runtime/VectorVM/Private/VectorVM.cpp b/Engine/Source/Runtime/VectorVM/Private/VectorVM.cpp
--- a/Engine/Source/Runtime/VectorVM/Private/VectorVM.cpp
+++ b/Engine/Source/Runtime/VectorVM/Private/VectorVM.cpp
@@ -35,6 +35,24 @@ DEFINE_LOG_CATEGORY_STATIC(LogVectorVM, All, All);
 
 
 
+UNiagaraDataObject::UNiagaraDataObject(const FObjectInitializer& ObjectInitializer)
+: Super(ObjectInitializer)
+{
+}
+
+UNiagaraCurveDataObject::UNiagaraCurveDataObject(const FObjectInitializer& ObjectInitializer)
+: Super(ObjectInitializer)
+{
+}
+
+UNiagaraSparseVolumeDataObject::UNiagaraSparseVolumeDataObject(const FObjectInitializer& ObjectInitializer)
+: Super(ObjectInitializer)
+{
+	Size = 64;
+	NumBuckets = Size*Size*Size;
+	//Data.AddZeroed(NumBuckets);
+	Data.Init(FVector4(0.1f, 0.1f, 0.1f, 0.1f), NumBuckets);
+}
 
 
 
@@ -51,7 +69,7 @@ struct FVectorVMContext
 	/** Pointer to the constant table. */
 	FVector4 const* RESTRICT ConstantTable;
 	/** Pointer to the data object constant table. */
-	FNiagaraDataObject * RESTRICT *DataObjConstantTable;
+	UNiagaraDataObject * RESTRICT *DataObjConstantTable;
 
 	/** The number of vectors to process. */
 	int32 NumVectors;
@@ -61,7 +79,7 @@ struct FVectorVMContext
 		const uint8* InCode,
 		VectorRegister** InRegisterTable,
 		const FVector4* InConstantTable,
-		FNiagaraDataObject** InDataObjTable,
+		UNiagaraDataObject** InDataObjTable,
 		int32 InNumVectors
 		)
 		: Code(InCode)
@@ -96,16 +114,16 @@ static VM_FORCEINLINE VectorRegister DecodeConstant(FVectorVMContext& Context)
 }
 
 /** Decode a constant from the bytecode. */
-static /*VM_FORCEINLINE*/ const FNiagaraDataObject *DecodeDataObjConstant(FVectorVMContext& Context)
+static /*VM_FORCEINLINE*/ const UNiagaraDataObject *DecodeDataObjConstant(FVectorVMContext& Context)
 {
-	const FNiagaraDataObject* Obj = Context.DataObjConstantTable[*Context.Code++];
+	const UNiagaraDataObject* Obj = Context.DataObjConstantTable[*Context.Code++];
 	return Obj;
 }
 
 /** Decode a constant from the bytecode. */
-static /*VM_FORCEINLINE*/ FNiagaraDataObject *DecodeWritableDataObjConstant(FVectorVMContext& Context)
+static /*VM_FORCEINLINE*/ UNiagaraDataObject *DecodeWritableDataObjConstant(FVectorVMContext& Context)
 {
-	FNiagaraDataObject* Obj = Context.DataObjConstantTable[*Context.Code++];
+	UNiagaraDataObject* Obj = Context.DataObjConstantTable[*Context.Code++];
 	return Obj;
 }
 
@@ -135,21 +153,21 @@ struct FConstantHandler
 /** Handles reading of a data object constant. */
 struct FDataObjectConstantHandler
 {
-	const FNiagaraDataObject *Constant;
+	const UNiagaraDataObject *Constant;
 	FDataObjectConstantHandler(FVectorVMContext& Context)
 		: Constant(DecodeDataObjConstant(Context))
 	{}
-	VM_FORCEINLINE const FNiagaraDataObject *Get(){ return Constant; }
+	VM_FORCEINLINE const UNiagaraDataObject *Get(){ return Constant; }
 };
 
 /** Handles reading of a data object constant. */
 struct FWritableDataObjectConstantHandler
 {
-	FNiagaraDataObject *Constant;
+	UNiagaraDataObject *Constant;
 	FWritableDataObjectConstantHandler(FVectorVMContext& Context)
 		: Constant(DecodeWritableDataObjConstant(Context))
 	{}
-	VM_FORCEINLINE FNiagaraDataObject *Get(){ return Constant; }
+	VM_FORCEINLINE UNiagaraDataObject *Get(){ return Constant; }
 };
 
 
@@ -599,7 +617,7 @@ struct FVectorKernelLessThan : public TBinaryVectorKernel<FVectorKernelLessThan>
 
 struct FVectorKernelSample : public TBinaryVectorKernelData<FVectorKernelSample>
 {
-	static void FORCEINLINE DoKernel(VectorRegister* RESTRICT Dst, const FNiagaraDataObject *Src0, VectorRegister Src1)
+	static void FORCEINLINE DoKernel(VectorRegister* RESTRICT Dst, const UNiagaraDataObject *Src0, VectorRegister Src1)
 	{
 		if (Src0)
 		{
@@ -612,7 +630,7 @@ struct FVectorKernelSample : public TBinaryVectorKernelData<FVectorKernelSample>
 
 struct FVectorKernelBufferWrite : public TTrinaryVectorKernelData<FVectorKernelBufferWrite>
 {
-	static void FORCEINLINE DoKernel(VectorRegister* RESTRICT Dst, FNiagaraDataObject *Src0, VectorRegister Src1, VectorRegister Src2)
+	static void FORCEINLINE DoKernel(VectorRegister* RESTRICT Dst, UNiagaraDataObject *Src0, VectorRegister Src1, VectorRegister Src2)
 	{
 		if (Src0)
 		{
@@ -672,6 +690,38 @@ struct FVectorKernelRandom : public TUnaryVectorKernel<FVectorKernelRandom>
 };
 
 
+/* gaussian distribution random number (not working yet) */
+struct FVectorKernelRandomGauss : public TBinaryVectorKernel<FVectorKernelRandomGauss>
+{
+	static void VM_FORCEINLINE DoKernel(VectorRegister* RESTRICT Dst, VectorRegister Src0, VectorRegister Src1)
+	{
+		const float rm = RAND_MAX;
+		VectorRegister Result = MakeVectorRegister(static_cast<float>(FMath::Rand()) / rm,
+			static_cast<float>(FMath::Rand()) / rm,
+			static_cast<float>(FMath::Rand()) / rm,
+			static_cast<float>(FMath::Rand()) / rm);
+
+		Result = VectorSubtract(Result, MakeVectorRegister(0.5f, 0.5f, 0.5f, 0.5f));
+		Result = VectorMultiply(MakeVectorRegister(3.0f, 3.0f, 3.0f, 3.0f), Result);
+
+		// taylor series gaussian approximation
+		const VectorRegister SPi2 = VectorReciprocal(VectorReciprocalSqrt(MakeVectorRegister(2 * PI, 2 * PI, 2 * PI, 2 * PI)));
+		VectorRegister Gauss = VectorReciprocal(SPi2);
+		VectorRegister Div = VectorMultiply(MakeVectorRegister(2.0f, 2.0f, 2.0f, 2.0f), SPi2);
+		Gauss = VectorSubtract(Gauss, VectorDivide(VectorMultiply(Result, Result), Div));
+		Div = VectorMultiply(MakeVectorRegister(8.0f, 8.0f, 8.0f, 8.0f), SPi2);
+		Gauss = VectorAdd(Gauss, VectorDivide(VectorPow(MakeVectorRegister(4.0f, 4.0f, 4.0f, 4.0f), Result), Div));
+		Div = VectorMultiply(MakeVectorRegister(48.0f, 48.0f, 48.0f, 48.0f), SPi2);
+		Gauss = VectorSubtract(Gauss, VectorDivide(VectorPow(MakeVectorRegister(6.0f, 6.0f, 6.0f, 6.0f), Result), Div));
+
+		Gauss = VectorDivide(Gauss, MakeVectorRegister(0.4f, 0.4f, 0.4f, 0.4f));
+		Gauss = VectorMultiply(Gauss, Src0);
+		*Dst = Gauss;
+	}
+};
+
+
+
 
 struct FVectorKernelMin : public TBinaryVectorKernel<FVectorKernelMin>
 {
@@ -919,7 +969,7 @@ void VectorVM::Exec(
 	VectorRegister** OutputRegisters,
 	int32 NumOutputRegisters,
 	FVector4 const* ConstantTable,
-	FNiagaraDataObject* *DataObjConstTable,
+	UNiagaraDataObject* *DataObjConstTable,
 	int32 NumVectors
 	)
 {
diff --git a/Engine/Source/Runtime/VectorVM/Public/VectorVM.h b/Engine/Source/Runtime/VectorVM/Public/VectorVM.h
--- a/Engine/Source/Runtime/VectorVM/Public/VectorVM.h
+++ b/Engine/Source/Runtime/VectorVM/Public/VectorVM.h
@@ -4,7 +4,7 @@
 #include "Core.h"
 #include "CoreUObject.h"
 
-class FNiagaraDataObject;
+class UNiagaraDataObject;
 
 namespace VectorVM
 {
@@ -108,7 +108,7 @@ namespace VectorVM
 		VectorRegister** OutputRegisters,
 		int32 NumOutputRegisters,
 		FVector4 const* ConstantTable,
-		FNiagaraDataObject* *DataObjTable,
+		UNiagaraDataObject* *DataObjTable,
 		int32 NumVectors
 		);
 
diff --git a/Engine/Source/Runtime/VectorVM/Public/VectorVMDataObject.h b/Engine/Source/Runtime/VectorVM/Public/VectorVMDataObject.h
--- a/Engine/Source/Runtime/VectorVM/Public/VectorVMDataObject.h
+++ b/Engine/Source/Runtime/VectorVM/Public/VectorVMDataObject.h
@@ -1,37 +1,33 @@
 // Copyright 1998-2015 Epic Games, Inc. All Rights Reserved.
 
 #pragma once	
+#include "VectorVMDataObject.generated.h"
 
 
 
 /* Vector VM data object; encapsulates buffers, curves and other data in its derivatives
 *  for access by VectorVM kernels;
 */
-class FNiagaraDataObject
+UCLASS()
+class UNiagaraDataObject : public UObject
 {
+	GENERATED_UCLASS_BODY()
 public:
-	virtual FVector4 Sample(const FVector4& InCoords) const = 0;
-	virtual FVector4 Write(const FVector4& InCoords, const FVector4& InValue) = 0;
-
-	virtual void Serialize(FArchive &Ar)
-	{
-
-	}
-	friend FArchive& operator<<(FArchive& Ar, FNiagaraDataObject* M)
-	{
-		M->Serialize(Ar);
-		return Ar;
-	}
+	virtual FVector4 Sample(const FVector4& InCoords) const { return FVector4(0.0f, 0.0f, 0.0f, 0.0f); }
+	virtual FVector4 Write(const FVector4& InCoords, const FVector4& InValue) { return FVector4(0.0f, 0.0f, 0.0f, 0.0f); }
 };
 
 /* Curve object; encapsulates a curve for the VectorVM
 */
-class FNiagaraCurveDataObject : public FNiagaraDataObject
+UCLASS(MinimalAPI)
+class UNiagaraCurveDataObject : public UNiagaraDataObject
 {
+	GENERATED_UCLASS_BODY()
 private:
+	UPROPERTY()
 	class UCurveVector *CurveObj;
 public:
-	FNiagaraCurveDataObject(class UCurveVector *InCurve) : CurveObj(InCurve)
+	UNiagaraCurveDataObject(class UCurveVector *InCurve) : CurveObj(InCurve)
 	{
 	}
 
@@ -48,31 +44,20 @@ public:
 
 	UCurveVector *GetCurveObject()	{ return CurveObj; }
 	void SetCurveObject(UCurveVector *InCurve)	{ CurveObj = InCurve; }
-
-
-	virtual void Serialize(FArchive &Ar) override
-	{
-		Ar << CurveObj;
-	}
 };
 
 
 /* Curve object; encapsulates sparse volumetric data for the VectorVM
 */
-class FNiagaraSparseVolumeDataObject : public FNiagaraDataObject
+UCLASS(MinimalAPI, Transient)
+class UNiagaraSparseVolumeDataObject : public UNiagaraDataObject
 {
+	GENERATED_UCLASS_BODY()
 private:
 	TArray<FVector4> Data;
 	uint64 NumBuckets;
 	int32 Size;
 public:
-	FNiagaraSparseVolumeDataObject()
-	{
-		Size = 64;
-		NumBuckets = Size*Size*Size;
-		Data.AddZeroed(NumBuckets);
-	}
-
 	uint32 MakeHash(const FVector4& InCoords) const
 	{
 		FVector4 Coords = InCoords;
diff --git a/Engine/Source/Runtime/VectorVM/VectorVM.Build.cs b/Engine/Source/Runtime/VectorVM/VectorVM.Build.cs
--- a/Engine/Source/Runtime/VectorVM/VectorVM.Build.cs
+++ b/Engine/Source/Runtime/VectorVM/VectorVM.Build.cs
@@ -6,9 +6,17 @@ public class VectorVM : ModuleRules
 {
 	public VectorVM(TargetInfo Target)
 	{
+        PublicDependencyModuleNames.AddRange(
+            new string[] {
+                "Core",
+                "Engine"
+            }
+        );
+
 		PrivateDependencyModuleNames.AddRange(
             new string[] {
                 "Core",
+				"CoreUObject", 
                 "Engine"
             }
         );
