Hopefully fixing a compiler error on the build machine that never showed up in my local builds.

diff --git a/Engine/Source/Runtime/VectorVM/Private/VectorVM.cpp b/Engine/Source/Runtime/VectorVM/Private/VectorVM.cpp
--- a/Engine/Source/Runtime/VectorVM/Private/VectorVM.cpp
+++ b/Engine/Source/Runtime/VectorVM/Private/VectorVM.cpp
@@ -45,6 +45,15 @@ UNiagaraCurveDataObject::UNiagaraCurveDataObject(const FObjectInitializer& Objec
 {
 }
 
+FVector4 UNiagaraCurveDataObject::Sample(const FVector4& InCoords) const
+{
+	FVector Vec = CurveObj->GetVectorValue(InCoords.X);
+	return FVector4(Vec, 0.0f);
+}
+
+
+
+
 UNiagaraSparseVolumeDataObject::UNiagaraSparseVolumeDataObject(const FObjectInitializer& ObjectInitializer)
 : Super(ObjectInitializer)
 {
diff --git a/Engine/Source/Runtime/VectorVM/Public/VectorVMDataObject.h b/Engine/Source/Runtime/VectorVM/Public/VectorVMDataObject.h
--- a/Engine/Source/Runtime/VectorVM/Public/VectorVMDataObject.h
+++ b/Engine/Source/Runtime/VectorVM/Public/VectorVMDataObject.h
@@ -4,7 +4,6 @@
 #include "VectorVMDataObject.generated.h"
 
 
-
 /* Vector VM data object; encapsulates buffers, curves and other data in its derivatives
 *  for access by VectorVM kernels;
 */
@@ -31,12 +30,7 @@ public:
 	{
 	}
 
-	FVector4 Sample(const FVector4& InCoords) const
-	{
-		FVector Vec = CurveObj->GetVectorValue(InCoords.X);
-		return FVector4(Vec, 0.0f);
-	}
-
+	FVector4 Sample(const FVector4& InCoords) const;
 	virtual FVector4 Write(const FVector4& InCoords, const FVector4& InValue) override
 	{
 		return FVector4(1.0f, 0.0f, 0.0f, 1.0f);
@@ -47,7 +41,7 @@ public:
 };
 
 
-/* Curve object; encapsulates sparse volumetric data for the VectorVM
+/* Volume data object; encapsulates volumetric data for VectorVM
 */
 UCLASS(MinimalAPI, Transient)
 class UNiagaraSparseVolumeDataObject : public UNiagaraDataObject
