** PRIMARY FEATURE **

diff --git a/Engine/Source/Runtime/VectorVM/Private/VectorVM.cpp b/Engine/Source/Runtime/VectorVM/Private/VectorVM.cpp
--- a/Engine/Source/Runtime/VectorVM/Private/VectorVM.cpp
+++ b/Engine/Source/Runtime/VectorVM/Private/VectorVM.cpp
@@ -45,16 +45,16 @@ DEFINE_LOG_CATEGORY_STATIC(LogVectorVM, All, All);
 
 #define SRCOP_RRCB (OP_REGISTER | OP_REGISTER | OP1_CONST | OP0_DATAOBJ)
 
-uint8 VectorVM::CreateSrcOperandMask(VectorVM::EOperandType Type1, VectorVM::EOperandType Type2, VectorVM::EOperandType Type3, VectorVM::EOperandType Type4)
-{
-	return	(Type1 == VectorVM::ConstantOperandType ? OP0_CONST : OP_REGISTER) |
-		(Type2 == VectorVM::ConstantOperandType ? OP1_CONST : OP_REGISTER) |
-		(Type3 == VectorVM::ConstantOperandType ? OP2_CONST : OP_REGISTER) |
-		(Type4 == VectorVM::ConstantOperandType ? OP3_CONST : OP_REGISTER) |
-		(Type1 == VectorVM::DataObjConstantOperandType ? OP0_DATAOBJ : OP_REGISTER) |
-		(Type2 == VectorVM::DataObjConstantOperandType ? OP1_DATAOBJ : OP_REGISTER) |
-		(Type3 == VectorVM::DataObjConstantOperandType ? OP2_DATAOBJ : OP_REGISTER) |
-		(Type4 == VectorVM::DataObjConstantOperandType ? OP3_DATAOBJ : OP_REGISTER);
+uint8 VectorVM::CreateSrcOperandMask(EVectorVMOperandLocation Type1, EVectorVMOperandLocation Type2, EVectorVMOperandLocation Type3, EVectorVMOperandLocation Type4)
+{
+	return	(Type1 == EVectorVMOperandLocation::Constant ? OP0_CONST : OP_REGISTER) |
+		(Type2 == EVectorVMOperandLocation::Constant ? OP1_CONST : OP_REGISTER) |
+		(Type3 == EVectorVMOperandLocation::Constant ? OP2_CONST : OP_REGISTER) |
+		(Type4 == EVectorVMOperandLocation::Constant ? OP3_CONST : OP_REGISTER) |
+		(Type1 == EVectorVMOperandLocation::DataObjConstant ? OP0_DATAOBJ : OP_REGISTER) |
+		(Type2 == EVectorVMOperandLocation::DataObjConstant ? OP1_DATAOBJ : OP_REGISTER) |
+		(Type3 == EVectorVMOperandLocation::DataObjConstant ? OP2_DATAOBJ : OP_REGISTER) |
+		(Type4 == EVectorVMOperandLocation::DataObjConstant ? OP3_DATAOBJ : OP_REGISTER);
 }
 
 
@@ -74,9 +74,6 @@ FVector4 UNiagaraCurveDataObject::Sample(const FVector4& InCoords) const
 	return FVector4(Vec, 0.0f);
 }
 
-
-
-
 UNiagaraSparseVolumeDataObject::UNiagaraSparseVolumeDataObject(const FObjectInitializer& ObjectInitializer)
 : Super(ObjectInitializer)
 {
@@ -86,8 +83,18 @@ UNiagaraSparseVolumeDataObject::UNiagaraSparseVolumeDataObject(const FObjectInit
 	Data.Init(FVector4(0.1f, 0.1f, 0.1f, 0.1f), NumBuckets);
 }
 
+//////////////////////////////////////////////////////////////////////////
 
+FDummyHandler DummyHandler;
 
+#if WITH_EDITOR
+TArray<FString> OpNames;
+TArray<FString> OperandLocationNames;
+#endif
+
+#if ENABLE_VM_DEBUGGING
+FVectorVMDebuggerImpl* AttachedDebugger = NULL;
+#endif
 
 /**
  * Context information passed around during VM execution.
@@ -102,198 +109,385 @@ struct FVectorVMContext
 	FVector4 const* RESTRICT ConstantTable;
 	/** Pointer to the data object constant table. */
 	UNiagaraDataObject * RESTRICT *DataObjConstantTable;
+	/** Pointer to the shared data table. */
+	FVectorVMSharedDataView* RESTRICT SharedDataTable;
 
 	/** The number of vectors to process. */
 	int32 NumVectors;
 
+	/** The number of instances to process. */
+	int32 NumInstances;
+
+	/** The Operation currently executing. */
+	EVectorVMOp CurrOp;
+
+	/** The instance we're currently starting at. Advances with each chunk processed. */
+	int32 StartInstance;
 	/** Initialization constructor. */
 	FVectorVMContext(
 		const uint8* InCode,
 		VectorRegister** InRegisterTable,
 		const FVector4* InConstantTable,
 		UNiagaraDataObject** InDataObjTable,
-		int32 InNumVectors
+		FVectorVMSharedDataView* InSharedDataTable,
+		int32 InNumVectors,
+		int32 InNumInstances,
+		int32 InStartInstance
 		)
 		: Code(InCode)
 		, RegisterTable(InRegisterTable)
 		, ConstantTable(InConstantTable)
 		, DataObjConstantTable(InDataObjTable)
+		, SharedDataTable(InSharedDataTable)
 		, NumVectors(InNumVectors)
+		, NumInstances(InNumInstances)
+		, CurrOp(EVectorVMOp::done)
+		, StartInstance(InStartInstance)
 	{
 	}
-};
 
+	FORCEINLINE bool IsDebugging()
+	{
+#if ENABLE_VM_DEBUGGING
+		return AttachedDebugger != nullptr; 
+#else
+		return false;
+#endif
+	}
+	FORCEINLINE void SetOp(EVectorVMOp InOp){ CurrOp = InOp; }
 
+#if ENABLE_VM_DEBUGGING
+	FORCEINLINE void BeginOp(VectorVM::EVMType InType, int32 InNumArgs, int32 InNumInstancesPerOp)
+	{
+		if (AttachedDebugger)
+		{
+			AttachedDebugger->BeginOp(*this, InType, InNumArgs, InNumInstancesPerOp);
+		}
+	}
 
+	template<typename DstHandler, typename Arg0Handler, typename Arg1Handler = FDummyHandler, typename Arg2Handler = FDummyHandler, typename Arg3Handler = FDummyHandler>
+	FORCEINLINE void PreOp(DstHandler& Dst, Arg0Handler& Arg0, Arg1Handler& Arg1 = DummyHandler, Arg2Handler& Arg2 = DummyHandler, Arg3Handler& Arg3 = DummyHandler)
+	{
+		if (AttachedDebugger)
+		{
+			AttachedDebugger->PreOp(*this, Dst, Arg0, Arg1, Arg2, Arg3);
+		}
+	}
 
-/** Decode the next operation contained in the bytecode. */
-static VM_FORCEINLINE VectorVM::EOp DecodeOp(FVectorVMContext& Context)
-{
-	return static_cast<VectorVM::EOp>(*Context.Code++);
-}
+	template<typename DstHandler, typename Arg0Handler, typename Arg1Handler = FDummyHandler, typename Arg2Handler = FDummyHandler, typename Arg3Handler = FDummyHandler>
+	FORCEINLINE void PostOp(DstHandler& Dst, Arg0Handler& Arg0, Arg1Handler& Arg1 = DummyHandler, Arg2Handler& Arg2 = DummyHandler, Arg3Handler& Arg3 = DummyHandler)
+	{
+		if (AttachedDebugger)
+		{
+			AttachedDebugger->PostOp(*this, Dst, Arg0, Arg1, Arg2, Arg3);
+		}
+	}
+#else
+	FORCEINLINE void BeginOp(VectorVM::EVMType InType, int32 InNumArgs, int32 InNumInstancesPerOp) { }
+
+	template<typename DstHandler, typename Arg0Handler, typename Arg1Handler = FDummyHandler, typename Arg2Handler = FDummyHandler, typename Arg3Handler = FDummyHandler>
+	FORCEINLINE void PreOp(DstHandler& Dst, Arg0Handler& Arg0, Arg1Handler& Arg1 = DummyHandler, Arg2Handler& Arg2 = DummyHandler, Arg3Handler& Arg3 = DummyHandler){	}
 
-/** Decode a register from the bytecode. */
-static VM_FORCEINLINE VectorRegister* DecodeRegister(FVectorVMContext& Context)
+	template<typename DstHandler, typename Arg0Handler, typename Arg1Handler = FDummyHandler, typename Arg2Handler = FDummyHandler, typename Arg3Handler = FDummyHandler>
+	FORCEINLINE void PostOp(DstHandler& Dst, Arg0Handler& Arg0, Arg1Handler& Arg1 = DummyHandler, Arg2Handler& Arg2 = DummyHandler, Arg3Handler& Arg3 = DummyHandler){	}
+#endif
+};
+
+
+static VM_FORCEINLINE uint8 DecodeU8(FVectorVMContext& Context)
 {
-	return Context.RegisterTable[*Context.Code++];
+	return *Context.Code++;
 }
 
-/** Decode a constant from the bytecode. */
-static VM_FORCEINLINE VectorRegister DecodeConstant(FVectorVMContext& Context)
+static VM_FORCEINLINE uint8 DecodeU16(FVectorVMContext& Context)
 {
-	const FVector4* vec = &Context.ConstantTable[*Context.Code++];
-	return VectorLoad(vec);
+	return (*((uint16*)Context.Code))++;
 }
 
-/** Decode a constant from the bytecode. */
-static /*VM_FORCEINLINE*/ const UNiagaraDataObject *DecodeDataObjConstant(FVectorVMContext& Context)
+static VM_FORCEINLINE uint8 DecodeU32(FVectorVMContext& Context)
 {
-	const UNiagaraDataObject* Obj = Context.DataObjConstantTable[*Context.Code++];
-	return Obj;
+	return (*((uint32*)Context.Code))++;
 }
 
-/** Decode a constant from the bytecode. */
-static /*VM_FORCEINLINE*/ UNiagaraDataObject *DecodeWritableDataObjConstant(FVectorVMContext& Context)
+/** Decode the next operation contained in the bytecode. */
+static VM_FORCEINLINE EVectorVMOp DecodeOp(FVectorVMContext& Context)
 {
-	UNiagaraDataObject* Obj = Context.DataObjConstantTable[*Context.Code++];
-	return Obj;
+	return static_cast<EVectorVMOp>(DecodeU8(Context));
 }
 
-
 static VM_FORCEINLINE uint8 DecodeSrcOperandTypes(FVectorVMContext& Context)
 {
-	return *Context.Code++;
+	return DecodeU8(Context);
 }
 
-static VM_FORCEINLINE uint8 DecodeByte(FVectorVMContext& Context)
+//////////////////////////////////////////////////////////////////////////
+/** Constant handler. */
+
+struct FConstantHandlerBase
 {
-	return *Context.Code++;
-}
+	uint8 ConstantIndex;
+	FConstantHandlerBase(FVectorVMContext& Context)
+		: ConstantIndex(DecodeU8(Context))
+	{}
 
+	VM_FORCEINLINE void Advance(){ }
+	VM_FORCEINLINE int32 GetLocationIndex(){ return ConstantIndex; }
+	VM_FORCEINLINE int32 GetSecondaryIndex(){ return INDEX_NONE; }
+	VM_FORCEINLINE int32 GetTertiaryIndex(){ return INDEX_NONE; }
+	VM_FORCEINLINE EVectorVMOperandLocation GetLocation(){ return EVectorVMOperandLocation::Constant; }
+};
 
+template<typename T>
+struct FConstantHandler : public FConstantHandlerBase
+{
+	T Constant;
+	FConstantHandler(FVectorVMContext& Context)
+		: FConstantHandlerBase(Context)
+		, Constant(Context.ConstantTable[ConstantIndex])
+	{}
+	VM_FORCEINLINE const T& Get(){ return Constant; }
+};
 
-/** Handles reading of a constant. */
-struct FConstantHandler
+template<>
+struct FConstantHandler<VectorRegister> : public FConstantHandlerBase
 {
 	VectorRegister Constant;
 	FConstantHandler(FVectorVMContext& Context)
-		: Constant(DecodeConstant(Context))
+		: FConstantHandlerBase(Context)
+		, Constant(VectorLoadAligned(&Context.ConstantTable[ConstantIndex]))
 	{}
-	VM_FORCEINLINE const VectorRegister& Get(){ return Constant; }
+	VM_FORCEINLINE const VectorRegister Get(){ return Constant; }
 };
 
-/** Handles reading of a data object constant. */
+typedef FConstantHandler<VectorRegister> FVectorConstantHandler;
+
+//////////////////////////////////////////////////////////////////////////
 struct FDataObjectConstantHandler
 {
-	const UNiagaraDataObject *Constant;
+	int32 ConstantIndex;
+	UNiagaraDataObject *Constant;
 	FDataObjectConstantHandler(FVectorVMContext& Context)
-		: Constant(DecodeDataObjConstant(Context))
+		: ConstantIndex(DecodeU8(Context))
+		, Constant(Context.DataObjConstantTable[ConstantIndex])
 	{}
-	VM_FORCEINLINE const UNiagaraDataObject *Get(){ return Constant; }
+	VM_FORCEINLINE void Advance(){ }
+	VM_FORCEINLINE int32 GetLocationIndex(){ return ConstantIndex; }
+	VM_FORCEINLINE int32 GetSecondaryIndex(){ return INDEX_NONE; }
+	VM_FORCEINLINE int32 GetTertiaryIndex(){ return INDEX_NONE; }
+	VM_FORCEINLINE EVectorVMOperandLocation GetLocation(){ return EVectorVMOperandLocation::DataObjConstant; }
+	VM_FORCEINLINE UNiagaraDataObject *Get(){ return Constant; }
 };
 
-/** Handles reading of a data object constant. */
-struct FWritableDataObjectConstantHandler
+//////////////////////////////////////////////////////////////////////////
+// Register handlers.
+// Handle reading of a register, advancing the pointer with each read.
+
+struct FRegisterHandlerBase
 {
-	UNiagaraDataObject *Constant;
-	FWritableDataObjectConstantHandler(FVectorVMContext& Context)
-		: Constant(DecodeWritableDataObjConstant(Context))
+	int32 RegisterIndex;
+	FRegisterHandlerBase(FVectorVMContext& Context)
+		: RegisterIndex(DecodeU8(Context))
 	{}
-	VM_FORCEINLINE UNiagaraDataObject *Get(){ return Constant; }
+	VM_FORCEINLINE int32 GetLocationIndex()
+	{
+		if (RegisterIndex < VectorVM::NumTempRegisters)
+		{
+			return RegisterIndex;
+		}
+		else if (RegisterIndex < VectorVM::NumTempRegisters + VectorVM::MaxInputRegisters)
+		{
+			return RegisterIndex - VectorVM::NumTempRegisters;
+		}
+		else
+		{
+			return RegisterIndex - (VectorVM::NumTempRegisters + VectorVM::MaxInputRegisters);
+		}
+	}
+	VM_FORCEINLINE int32 GetSecondaryIndex(){ return INDEX_NONE; }
+	VM_FORCEINLINE int32 GetTertiaryIndex(){ return INDEX_NONE; }
+	VM_FORCEINLINE EVectorVMOperandLocation GetLocation()
+	{
+		if (RegisterIndex < VectorVM::NumTempRegisters)
+		{
+			return EVectorVMOperandLocation::TemporaryRegister;
+		}
+		else if (RegisterIndex < VectorVM::NumTempRegisters + VectorVM::MaxInputRegisters)
+		{
+			return EVectorVMOperandLocation::InputRegister;
+		}
+		else
+		{ 
+			return EVectorVMOperandLocation::OutputRegister;
+		}
+	}
 };
 
+template<typename T, int32 NumInstancesPerOp=1>
+struct FRegisterHandler : public FRegisterHandlerBase
+{
+	T* RESTRICT Register;
+	FRegisterHandler(FVectorVMContext& Context)
+		: FRegisterHandlerBase(Context)
+		, Register((T*)Context.RegisterTable[RegisterIndex])
+	{}
+	VM_FORCEINLINE const T Get(){ return *Register; }
+	VM_FORCEINLINE void Advance(){ Register += NumInstancesPerOp; }
+};
 
-/** Handles reading of a register, advancing the pointer with each read. */
-struct FRegisterHandler
+template<int32 NumInstancesPerOp> struct FRegisterHandler<VectorRegister, NumInstancesPerOp> : public FRegisterHandlerBase
 {
-	VectorRegister* Register;
-	VectorRegister RegisterValue;
+	VectorRegister* RESTRICT Register;
 	FRegisterHandler(FVectorVMContext& Context)
-		: Register(DecodeRegister(Context))
+		: FRegisterHandlerBase(Context)
+		, Register((VectorRegister*)Context.RegisterTable[RegisterIndex])
 	{}
-	VM_FORCEINLINE const VectorRegister& Get(){ RegisterValue = *Register++;  return RegisterValue; }
+	VM_FORCEINLINE const VectorRegister Get(){ return VectorLoadAligned(Register); }
+	VM_FORCEINLINE void Advance(){ Register += NumInstancesPerOp; }
 };
 
-template<typename Kernel, typename Arg0Handler>
-VM_FORCEINLINE void VectorUnaryLoop(FVectorVMContext& Context, VectorRegister* RESTRICT Dst, int32 NumVectors)
+typedef FRegisterHandler<VectorRegister> FVectorRegisterHandler;
+
+/** Handles writing to a register, advancing the pointer with each write. */
+template<typename T, int32 NumInstancesPerOp = 1>
+struct FRegisterDestHandler : public FRegisterHandlerBase
+{
+	T* RESTRICT Register;
+	FRegisterDestHandler(FVectorVMContext& Context)
+		: FRegisterHandlerBase(Context)
+		, Register((T*)Context.RegisterTable[RegisterIndex])
+	{}
+	VM_FORCEINLINE T* RESTRICT Get(){ return Register; }
+	VM_FORCEINLINE T GetValue(){ return *Register; }
+	VM_FORCEINLINE void Advance(){ Register += NumInstancesPerOp; }
+};
+
+template<int32 NumInstancesPerOp>
+struct FRegisterDestHandler<VectorRegister, NumInstancesPerOp> : public FRegisterHandlerBase
+{
+	VectorRegister* RESTRICT Register;
+	FRegisterDestHandler(FVectorVMContext& Context)
+		: FRegisterHandlerBase(Context)
+		, Register((VectorRegister*)Context.RegisterTable[RegisterIndex])
+	{}
+	VM_FORCEINLINE VectorRegister* RESTRICT Get(){ return Register; }
+	VM_FORCEINLINE VectorRegister GetValue(){ return VectorLoadAligned(Register); }
+	VM_FORCEINLINE void Advance(){ Register += NumInstancesPerOp; }
+};
+
+//////////////////////////////////////////////////////////////////////////
+// Kernels
+template<typename Kernel, typename DstHandler, typename Arg0Handler>
+struct TUnaryKernel
 {
-	Arg0Handler Arg0(Context);
-	for (int32 i = 0; i < NumVectors; ++i)
+	static VM_FORCEINLINE void Exec(FVectorVMContext& Context)
 	{
-		Kernel::DoKernel(Dst++, Arg0.Get());
+		DstHandler Dst(Context);
+		Arg0Handler Arg0(Context);
+
+		Context.BeginOp(Kernel::Type, 1, Kernel::NumInstancesPerOp);
+		for (int32 i = 0; i < Context.NumInstances; i += Kernel::NumInstancesPerOp)
+		{
+			Context.PreOp(Dst, Arg0);
+			Kernel::DoKernel(Dst.Get(), Arg0.Get());
+			Context.PostOp(Dst, Arg0);
+			Dst.Advance();	Arg0.Advance();
+		}
 	}
-}
+};
 
-template<typename Kernel, typename Arg0Handler, typename Arg1Handler>
-VM_FORCEINLINE void VectorBinaryLoop(FVectorVMContext& Context, VectorRegister* RESTRICT Dst, int32 NumVectors)
+template<typename Kernel, typename DstHandler, typename Arg0Handler, typename Arg1Handler>
+struct TBinaryKernel
 {
-	Arg0Handler Arg0(Context);
-	Arg1Handler Arg1(Context);
-	for (int32 i = 0; i < NumVectors; ++i)
+	static VM_FORCEINLINE void Exec(FVectorVMContext& Context)
 	{
-		Kernel::DoKernel(Dst++, Arg0.Get(), Arg1.Get());
+		DstHandler Dst(Context); 
+		Arg0Handler Arg0(Context); 
+		Arg1Handler Arg1(Context);
+		Context.BeginOp(Kernel::Type, 2, Kernel::NumInstancesPerOp);
+		for (int32 i = 0; i < Context.NumInstances; i += Kernel::NumInstancesPerOp)
+		{
+			Context.PreOp(Dst, Arg0, Arg1);
+			Kernel::DoKernel(Dst.Get(), Arg0.Get(), Arg1.Get());
+			Context.PostOp(Dst, Arg0, Arg1);
+			Dst.Advance(); Arg0.Advance(); Arg1.Advance();
+		}
 	}
-}
-
+};
 
-template<typename Kernel, typename Arg0Handler, typename Arg1Handler, typename Arg2Handler>
-VM_FORCEINLINE void VectorTrinaryLoop(FVectorVMContext& Context, VectorRegister* RESTRICT Dst, int32 NumVectors)
+template<typename Kernel, typename DstHandler, typename Arg0Handler, typename Arg1Handler, typename Arg2Handler>
+struct TTrinaryKernel
 {
-	Arg0Handler Arg0(Context);
-	Arg1Handler Arg1(Context);
-	Arg2Handler Arg2(Context);
-	for (int32 i = 0; i < NumVectors; ++i)
+	static VM_FORCEINLINE void Exec(FVectorVMContext& Context)
 	{
-		Kernel::DoKernel(Dst++, Arg0.Get(), Arg1.Get(), Arg2.Get());
+		DstHandler Dst(Context);
+		Arg0Handler Arg0(Context);
+		Arg1Handler Arg1(Context);
+		Arg2Handler Arg2(Context);
+		Context.BeginOp(Kernel::Type, 3, Kernel::NumInstancesPerOp);
+		for (int32 i = 0; i < Context.NumInstances; i += Kernel::NumInstancesPerOp)
+		{
+			Context.PreOp(Dst, Arg0, Arg1, Arg2);
+			Kernel::DoKernel(Dst.Get(), Arg0.Get(), Arg1.Get(), Arg2.Get());
+			Context.PostOp(Dst, Arg0, Arg1, Arg2);
+			Dst.Advance(); Arg0.Advance(); Arg1.Advance(); Arg2.Advance();
+		}
 	}
-}
+};
 
-template<typename Kernel, typename Arg0Handler, typename Arg1Handler, typename Arg2Handler, typename Arg3Handler>
-VM_FORCEINLINE void VectorQuatenaryLoop(FVectorVMContext& Context, VectorRegister* RESTRICT Dst, int32 NumVectors)
+template<typename Kernel, typename DstHandler, typename Arg0Handler, typename Arg1Handler, typename Arg2Handler, typename Arg3Handler>
+struct TQuaternaryKernel
 {
-	Arg0Handler Arg0(Context);
-	Arg1Handler Arg1(Context);
-	Arg2Handler Arg2(Context);
-	Arg3Handler Arg3(Context);
-	for (int32 i = 0; i < NumVectors; ++i)
+	static VM_FORCEINLINE void Exec(FVectorVMContext& Context)
 	{
-		Kernel::DoKernel(Dst++, Arg0.Get(), Arg1.Get(), Arg2.Get(), Arg3.Get());
+		DstHandler Dst(Context);
+		Arg0Handler Arg0(Context);
+		Arg1Handler Arg1(Context);
+		Arg2Handler Arg2(Context);
+		Arg3Handler Arg3(Context);
+		Context.BeginOp(Kernel::Type, 4, Kernel::NumInstancesPerOp);
+		for (int32 i = 0; i < Context.NumInstances; i += Kernel::NumInstancesPerOp)
+		{
+			Context.PreOp(Dst, Arg0, Arg1, Arg2);
+			Kernel::DoKernel(Dst.Get(), Arg0.Get(), Arg1.Get(), Arg2.Get(), Arg3.Get());
+			Context.PostOp(Dst, Arg0, Arg1, Arg2, Arg3);
+			Dst.Advance(); Arg0.Advance(); Arg1.Advance(); Arg2.Advance(); Arg3.Advance();
+		}
 	}
-}
+};
 
 /** Base class of vector kernels with a single operand. */
-template <typename Kernel>
+template <typename Kernel, typename DstHandler = FRegisterDestHandler<VectorRegister>>
 struct TUnaryVectorKernel
 {
+	static const VectorVM::EVMType Type = VectorVM::EVMType::Vector4;
+	static const int32 NumInstancesPerOp = 1;
 	static void Exec(FVectorVMContext& Context)
 	{
-		VectorRegister* RESTRICT Dst = DecodeRegister(Context);
 		uint32 SrcOpTypes = DecodeSrcOperandTypes(Context);
-		int32 NumVectors = Context.NumVectors;
 		switch (SrcOpTypes)
 		{
-		case SRCOP_RRRR: 	VectorUnaryLoop<Kernel, FRegisterHandler>(Context, Dst, NumVectors); break;
-		case SRCOP_RRRC:	VectorUnaryLoop<Kernel, FConstantHandler>(Context, Dst, NumVectors); break;
-		default: check(0); break;
+		case SRCOP_RRRR: 	TUnaryKernel<Kernel, DstHandler, FVectorRegisterHandler>::Exec(Context); break;
+		case SRCOP_RRRC:	TUnaryKernel<Kernel, DstHandler, FVectorConstantHandler>::Exec(Context); break;
+		default: check(0); break; 
 		};
 	}
 };
 
 /** Base class of Vector kernels with 2 operands. */
-template <typename Kernel>
+template <typename Kernel, typename DstHandler = FRegisterDestHandler<VectorRegister>>
 struct TBinaryVectorKernel
 {
+	static const VectorVM::EVMType Type = VectorVM::EVMType::Vector4;
+	static const int32 NumInstancesPerOp = 1;
 	static void Exec(FVectorVMContext& Context)
 	{
-		VectorRegister* RESTRICT Dst = DecodeRegister(Context);
 		uint32 SrcOpTypes = DecodeSrcOperandTypes(Context);
-		int32 NumVectors = Context.NumVectors;
 		switch (SrcOpTypes)
 		{
-		case SRCOP_RRRR: 	VectorBinaryLoop<Kernel, FRegisterHandler, FRegisterHandler>(Context, Dst, NumVectors); break;
-		case SRCOP_RRRC:	VectorBinaryLoop<Kernel, FConstantHandler, FRegisterHandler>(Context, Dst, NumVectors); break;
-		case SRCOP_RRCR: 	VectorBinaryLoop<Kernel, FRegisterHandler, FConstantHandler>(Context, Dst, NumVectors); break;
-		case SRCOP_RRCC:	VectorBinaryLoop<Kernel, FConstantHandler, FConstantHandler>(Context, Dst, NumVectors); break;
+		case SRCOP_RRRR: 	TBinaryKernel<Kernel, DstHandler, FVectorRegisterHandler, FVectorRegisterHandler>::Exec(Context); break;
+		case SRCOP_RRRC:	TBinaryKernel<Kernel, DstHandler, FVectorConstantHandler, FVectorRegisterHandler>::Exec(Context); break;
+		case SRCOP_RRCR: 	TBinaryKernel<Kernel, DstHandler, FVectorRegisterHandler, FVectorConstantHandler>::Exec(Context); break;
+		case SRCOP_RRCC:	TBinaryKernel<Kernel, DstHandler, FVectorConstantHandler, FVectorConstantHandler>::Exec(Context); break;
 		default: check(0); break;
 		};
 	}
@@ -303,18 +497,18 @@ struct TBinaryVectorKernel
 
 
 /** Base class of Vector kernels with 2 operands, one of which can be a data object. */
-template <typename Kernel>
+template <typename Kernel, typename DstHandler = FRegisterDestHandler<VectorRegister>>
 struct TBinaryVectorKernelData
 {
+	static const VectorVM::EVMType Type = VectorVM::EVMType::Vector4;
+	static const int32 NumInstancesPerOp = 1;
 	static void Exec(FVectorVMContext& Context)
 	{
-		VectorRegister* RESTRICT Dst = DecodeRegister(Context);
 		uint32 SrcOpTypes = DecodeSrcOperandTypes(Context);
-		int32 NumVectors = Context.NumVectors;
 		switch (SrcOpTypes)
 		{
-		case SRCOP_RRRB:	VectorBinaryLoop<Kernel, FDataObjectConstantHandler, FRegisterHandler>(Context, Dst, NumVectors); break;
-		case SRCOP_RRCB:	VectorBinaryLoop<Kernel, FDataObjectConstantHandler, FConstantHandler>(Context, Dst, NumVectors); break;
+		case SRCOP_RRRB:	TBinaryKernel<Kernel, DstHandler, FDataObjectConstantHandler, FVectorRegisterHandler>::Exec(Context); break;
+		case SRCOP_RRCB:	TBinaryKernel<Kernel, DstHandler, FDataObjectConstantHandler, FVectorConstantHandler>::Exec(Context); break;
 		default: check(0); break;
 		};
 	}
@@ -323,17 +517,17 @@ struct TBinaryVectorKernelData
 
 
 /** Base class of Vector kernels with 2 operands, one of which can be a data object. */
-template <typename Kernel>
+template <typename Kernel, typename DstHandler = FRegisterDestHandler<VectorRegister>>
 struct TTrinaryVectorKernelData
 {
+	static const VectorVM::EVMType Type = VectorVM::EVMType::Vector4;
+	static const int32 NumInstancesPerOp = 1;
 	static void Exec(FVectorVMContext& Context)
 	{
-		VectorRegister* RESTRICT Dst = DecodeRegister(Context);
 		uint32 SrcOpTypes = DecodeSrcOperandTypes(Context);
-		int32 NumVectors = Context.NumVectors;
 		switch (SrcOpTypes)
 		{
-		case SRCOP_RRRB:	VectorTrinaryLoop<Kernel, FWritableDataObjectConstantHandler, FRegisterHandler, FRegisterHandler>(Context, Dst, NumVectors); break;
+		case SRCOP_RRRB:	TTrinaryKernel<Kernel, DstHandler, FDataObjectConstantHandler, FVectorRegisterHandler, FVectorRegisterHandler>::Exec(Context); break;
 		default: check(0); break;
 		};
 	}
@@ -343,56 +537,56 @@ struct TTrinaryVectorKernelData
 
 
 /** Base class of Vector kernels with 3 operands. */
-template <typename Kernel>
+template <typename Kernel, typename DstHandler = FRegisterDestHandler<VectorRegister>>
 struct TTrinaryVectorKernel
 {
+	static const VectorVM::EVMType Type = VectorVM::EVMType::Vector4;
+	static const int32 NumInstancesPerOp = 1;
 	static void Exec(FVectorVMContext& Context)
 	{
-		VectorRegister* RESTRICT Dst = DecodeRegister(Context);
 		uint32 SrcOpTypes = DecodeSrcOperandTypes(Context);
-		int32 NumVectors = Context.NumVectors;
 		switch (SrcOpTypes)
 		{
-		case SRCOP_RRRR: 	VectorTrinaryLoop<Kernel, FRegisterHandler, FRegisterHandler, FRegisterHandler>(Context, Dst, NumVectors); break;
-		case SRCOP_RRRC:	VectorTrinaryLoop<Kernel, FConstantHandler, FRegisterHandler, FRegisterHandler>(Context, Dst, NumVectors); break;
-		case SRCOP_RRCR: 	VectorTrinaryLoop<Kernel, FRegisterHandler, FConstantHandler, FRegisterHandler>(Context, Dst, NumVectors); break;
-		case SRCOP_RRCC:	VectorTrinaryLoop<Kernel, FConstantHandler, FConstantHandler, FRegisterHandler>(Context, Dst, NumVectors); break;
-		case SRCOP_RCRR: 	VectorTrinaryLoop<Kernel, FRegisterHandler, FRegisterHandler, FConstantHandler>(Context, Dst, NumVectors); break;
-		case SRCOP_RCRC:	VectorTrinaryLoop<Kernel, FConstantHandler, FRegisterHandler, FConstantHandler>(Context, Dst, NumVectors); break;
-		case SRCOP_RCCR: 	VectorTrinaryLoop<Kernel, FRegisterHandler, FConstantHandler, FConstantHandler>(Context, Dst, NumVectors); break;
-		case SRCOP_RCCC:	VectorTrinaryLoop<Kernel, FConstantHandler, FConstantHandler, FConstantHandler>(Context, Dst, NumVectors); break;
+		case SRCOP_RRRR: 	TTrinaryKernel<Kernel, DstHandler, FVectorRegisterHandler, FVectorRegisterHandler, FVectorRegisterHandler>::Exec(Context); break;
+		case SRCOP_RRRC:	TTrinaryKernel<Kernel, DstHandler, FVectorConstantHandler, FVectorRegisterHandler, FVectorRegisterHandler>::Exec(Context); break;
+		case SRCOP_RRCR: 	TTrinaryKernel<Kernel, DstHandler, FVectorRegisterHandler, FVectorConstantHandler, FVectorRegisterHandler>::Exec(Context); break;
+		case SRCOP_RRCC:	TTrinaryKernel<Kernel, DstHandler, FVectorConstantHandler, FVectorConstantHandler, FVectorRegisterHandler>::Exec(Context); break;
+		case SRCOP_RCRR: 	TTrinaryKernel<Kernel, DstHandler, FVectorRegisterHandler, FVectorRegisterHandler, FVectorConstantHandler>::Exec(Context); break;
+		case SRCOP_RCRC:	TTrinaryKernel<Kernel, DstHandler, FVectorConstantHandler, FVectorRegisterHandler, FVectorConstantHandler>::Exec(Context); break;
+		case SRCOP_RCCR: 	TTrinaryKernel<Kernel, DstHandler, FVectorRegisterHandler, FVectorConstantHandler, FVectorConstantHandler>::Exec(Context); break;
+		case SRCOP_RCCC:	TTrinaryKernel<Kernel, DstHandler, FVectorConstantHandler, FVectorConstantHandler, FVectorConstantHandler>::Exec(Context); break;
 		default: check(0); break;
 		};
 	}
 };
 
 /** Base class of Vector kernels with 4 operands. */
-template <typename Kernel>
+template <typename Kernel, typename DstHandler = FRegisterDestHandler<VectorRegister>>
 struct TQuatenaryVectorKernel
 {
+	static const VectorVM::EVMType Type = VectorVM::EVMType::Vector4;
+	static const int32 NumInstancesPerOp = 1;
 	static void Exec(FVectorVMContext& Context)
 	{
-		VectorRegister* RESTRICT Dst = DecodeRegister(Context);
 		uint32 SrcOpTypes = DecodeSrcOperandTypes(Context);
-		int32 NumVectors = Context.NumVectors;
 		switch (SrcOpTypes)
 		{
-		case SRCOP_RRRR: 	VectorQuatenaryLoop<Kernel, FRegisterHandler, FRegisterHandler, FRegisterHandler, FRegisterHandler>(Context, Dst, NumVectors); break;
-		case SRCOP_RRRC:	VectorQuatenaryLoop<Kernel, FConstantHandler, FRegisterHandler, FRegisterHandler, FRegisterHandler>(Context, Dst, NumVectors); break;
-		case SRCOP_RRCR: 	VectorQuatenaryLoop<Kernel, FRegisterHandler, FConstantHandler, FRegisterHandler, FRegisterHandler>(Context, Dst, NumVectors); break;
-		case SRCOP_RRCC:	VectorQuatenaryLoop<Kernel, FConstantHandler, FConstantHandler, FRegisterHandler, FRegisterHandler>(Context, Dst, NumVectors); break;
-		case SRCOP_RCRR: 	VectorQuatenaryLoop<Kernel, FRegisterHandler, FRegisterHandler, FConstantHandler, FRegisterHandler>(Context, Dst, NumVectors); break;
-		case SRCOP_RCRC:	VectorQuatenaryLoop<Kernel, FConstantHandler, FRegisterHandler, FConstantHandler, FRegisterHandler>(Context, Dst, NumVectors); break;
-		case SRCOP_RCCR: 	VectorQuatenaryLoop<Kernel, FRegisterHandler, FConstantHandler, FConstantHandler, FRegisterHandler>(Context, Dst, NumVectors); break;
-		case SRCOP_RCCC:	VectorQuatenaryLoop<Kernel, FConstantHandler, FConstantHandler, FConstantHandler, FRegisterHandler>(Context, Dst, NumVectors); break;
-		case SRCOP_CRRR: 	VectorQuatenaryLoop<Kernel, FRegisterHandler, FRegisterHandler, FRegisterHandler, FConstantHandler>(Context, Dst, NumVectors); break;
-		case SRCOP_CRRC:	VectorQuatenaryLoop<Kernel, FConstantHandler, FRegisterHandler, FRegisterHandler, FConstantHandler>(Context, Dst, NumVectors); break;
-		case SRCOP_CRCR: 	VectorQuatenaryLoop<Kernel, FRegisterHandler, FConstantHandler, FRegisterHandler, FConstantHandler>(Context, Dst, NumVectors); break;
-		case SRCOP_CRCC:	VectorQuatenaryLoop<Kernel, FConstantHandler, FConstantHandler, FRegisterHandler, FConstantHandler>(Context, Dst, NumVectors); break;
-		case SRCOP_CCRR: 	VectorQuatenaryLoop<Kernel, FRegisterHandler, FRegisterHandler, FConstantHandler, FConstantHandler>(Context, Dst, NumVectors); break;
-		case SRCOP_CCRC:	VectorQuatenaryLoop<Kernel, FConstantHandler, FRegisterHandler, FConstantHandler, FConstantHandler>(Context, Dst, NumVectors); break;
-		case SRCOP_CCCR: 	VectorQuatenaryLoop<Kernel, FRegisterHandler, FConstantHandler, FConstantHandler, FConstantHandler>(Context, Dst, NumVectors); break;
-		case SRCOP_CCCC:	VectorQuatenaryLoop<Kernel, FConstantHandler, FConstantHandler, FConstantHandler, FConstantHandler>(Context, Dst, NumVectors); break;
+		case SRCOP_RRRR: 	TQuaternaryKernel<Kernel, DstHandler, FVectorRegisterHandler, FVectorRegisterHandler, FVectorRegisterHandler, FVectorRegisterHandler>::Exec(Context); break;
+		case SRCOP_RRRC:	TQuaternaryKernel<Kernel, DstHandler, FVectorConstantHandler, FVectorRegisterHandler, FVectorRegisterHandler, FVectorRegisterHandler>::Exec(Context); break;
+		case SRCOP_RRCR: 	TQuaternaryKernel<Kernel, DstHandler, FVectorRegisterHandler, FVectorConstantHandler, FVectorRegisterHandler, FVectorRegisterHandler>::Exec(Context); break;
+		case SRCOP_RRCC:	TQuaternaryKernel<Kernel, DstHandler, FVectorConstantHandler, FVectorConstantHandler, FVectorRegisterHandler, FVectorRegisterHandler>::Exec(Context); break;
+		case SRCOP_RCRR: 	TQuaternaryKernel<Kernel, DstHandler, FVectorRegisterHandler, FVectorRegisterHandler, FVectorConstantHandler, FVectorRegisterHandler>::Exec(Context); break;
+		case SRCOP_RCRC:	TQuaternaryKernel<Kernel, DstHandler, FVectorConstantHandler, FVectorRegisterHandler, FVectorConstantHandler, FVectorRegisterHandler>::Exec(Context); break;
+		case SRCOP_RCCR: 	TQuaternaryKernel<Kernel, DstHandler, FVectorRegisterHandler, FVectorConstantHandler, FVectorConstantHandler, FVectorRegisterHandler>::Exec(Context); break;
+		case SRCOP_RCCC:	TQuaternaryKernel<Kernel, DstHandler, FVectorConstantHandler, FVectorConstantHandler, FVectorConstantHandler, FVectorRegisterHandler>::Exec(Context); break;
+		case SRCOP_CRRR: 	TQuaternaryKernel<Kernel, DstHandler, FVectorRegisterHandler, FVectorRegisterHandler, FVectorRegisterHandler, FVectorConstantHandler>::Exec(Context); break;
+		case SRCOP_CRRC:	TQuaternaryKernel<Kernel, DstHandler, FVectorConstantHandler, FVectorRegisterHandler, FVectorRegisterHandler, FVectorConstantHandler>::Exec(Context); break;
+		case SRCOP_CRCR: 	TQuaternaryKernel<Kernel, DstHandler, FVectorRegisterHandler, FVectorConstantHandler, FVectorRegisterHandler, FVectorConstantHandler>::Exec(Context); break;
+		case SRCOP_CRCC:	TQuaternaryKernel<Kernel, DstHandler, FVectorConstantHandler, FVectorConstantHandler, FVectorRegisterHandler, FVectorConstantHandler>::Exec(Context); break;
+		case SRCOP_CCRR: 	TQuaternaryKernel<Kernel, DstHandler, FVectorRegisterHandler, FVectorRegisterHandler, FVectorConstantHandler, FVectorConstantHandler>::Exec(Context); break;
+		case SRCOP_CCRC:	TQuaternaryKernel<Kernel, DstHandler, FVectorConstantHandler, FVectorRegisterHandler, FVectorConstantHandler, FVectorConstantHandler>::Exec(Context); break;
+		case SRCOP_CCCR: 	TQuaternaryKernel<Kernel, DstHandler, FVectorRegisterHandler, FVectorConstantHandler, FVectorConstantHandler, FVectorConstantHandler>::Exec(Context); break;
+		case SRCOP_CCCC:	TQuaternaryKernel<Kernel, DstHandler, FVectorConstantHandler, FVectorConstantHandler, FVectorConstantHandler, FVectorConstantHandler>::Exec(Context); break;
 		default: check(0); break;
 		};
 	}
@@ -447,8 +641,7 @@ struct FVectorKernelLerp : public TTrinaryVectorKernel<FVectorKernelLerp>
 {
 	static void VM_FORCEINLINE DoKernel(VectorRegister* RESTRICT Dst,VectorRegister Src0,VectorRegister Src1,VectorRegister Src2)
 	{
-		const VectorRegister One = MakeVectorRegister(1.0f, 1.0f, 1.0f, 1.0f);
-		const VectorRegister OneMinusAlpha = VectorSubtract(One, Src2);
+		const VectorRegister OneMinusAlpha = VectorSubtract(GlobalVectorConstants::FloatOne, Src2);
 		const VectorRegister Tmp = VectorMultiply(Src0, OneMinusAlpha);
 		*Dst = VectorMultiplyAdd(Src1, Src2, Tmp);
 	}
@@ -540,7 +733,7 @@ struct FVectorKernelSin : public TUnaryVectorKernel<FVectorKernelSin>
 {
 	static void VM_FORCEINLINE DoKernel(VectorRegister* RESTRICT Dst, VectorRegister Src0)
 	{
-		*Dst = VectorSin( VectorMultiply(Src0, GlobalVectorConstants::PiByTwo));
+		*Dst = VectorSin(VectorMultiply(Src0, GlobalVectorConstants::TwoPi));
 	}
 };
 
@@ -548,7 +741,7 @@ struct FVectorKernelCos : public TUnaryVectorKernel<FVectorKernelCos>
 {
 	static void VM_FORCEINLINE DoKernel(VectorRegister* RESTRICT Dst, VectorRegister Src0)
  	{
-		*Dst = VectorCos(VectorMultiply(Src0, GlobalVectorConstants::PiByTwo));
+		*Dst = VectorCos(VectorMultiply(Src0, GlobalVectorConstants::TwoPi));
 	}
 };
 
@@ -556,7 +749,7 @@ struct FVectorKernelTan : public TUnaryVectorKernel<FVectorKernelTan>
 {
 	static void VM_FORCEINLINE DoKernel(VectorRegister* RESTRICT Dst, VectorRegister Src0)
 	{
-		*Dst = VectorTan(VectorMultiply(Src0, GlobalVectorConstants::PiByTwo));
+		*Dst = VectorTan(VectorMultiply(Src0, GlobalVectorConstants::TwoPi));
 	}
 };
 
@@ -564,7 +757,7 @@ struct FVectorKernelASin : public TUnaryVectorKernel<FVectorKernelASin>
 {
 	static void VM_FORCEINLINE DoKernel(VectorRegister* RESTRICT Dst, VectorRegister Src0)
 	{
-		*Dst = VectorASin(Src0);
+		*Dst = VectorMultiply(VectorASin(Src0), GlobalVectorConstants::OneOverTwoPi);
 	}
 };
 
@@ -572,7 +765,7 @@ struct FVectorKernelACos : public TUnaryVectorKernel<FVectorKernelACos>
 {
 	static void VM_FORCEINLINE DoKernel(VectorRegister* RESTRICT Dst, VectorRegister Src0)
 	{
-		*Dst = VectorACos(Src0);
+		*Dst = VectorMultiply(VectorACos(Src0), GlobalVectorConstants::OneOverTwoPi);
 	}
 };
 
@@ -580,7 +773,7 @@ struct FVectorKernelATan : public TUnaryVectorKernel<FVectorKernelATan>
 {
 	static void VM_FORCEINLINE DoKernel(VectorRegister* RESTRICT Dst, VectorRegister Src0)
 	{
-		*Dst = VectorATan(Src0);
+		*Dst = VectorMultiply(VectorATan(Src0), GlobalVectorConstants::OneOverTwoPi);
 	}
 };
 
@@ -588,7 +781,7 @@ struct FVectorKernelATan2 : public TBinaryVectorKernel<FVectorKernelATan2>
 {
 	static void VM_FORCEINLINE DoKernel(VectorRegister* RESTRICT Dst, VectorRegister Src0, VectorRegister Src1)
 	{
-		*Dst = VectorATan2(Src0, Src1);
+		*Dst = VectorMultiply(VectorATan2(Src0, Src1), GlobalVectorConstants::OneOverTwoPi);
 	}
 };
 
@@ -636,21 +829,26 @@ struct FVectorKernelLessThan : public TBinaryVectorKernel<FVectorKernelLessThan>
 {
 	static void VM_FORCEINLINE DoKernel(VectorRegister* RESTRICT Dst, VectorRegister Src0, VectorRegister Src1)
 	{
-		const VectorRegister Large = MakeVectorRegister(BIG_NUMBER, BIG_NUMBER, BIG_NUMBER, BIG_NUMBER);
-		const VectorRegister One = MakeVectorRegister(1.0f, 1.0f, 1.0f, 1.0f);
-		const VectorRegister Zero = MakeVectorRegister(0.0f, 0.0f, 0.0f, 0.0f);
 		VectorRegister Tmp = VectorSubtract(Src1, Src0);
-		Tmp = VectorMultiply(Tmp, Large);
-		Tmp = VectorMin(Tmp, One);
-		*Dst = VectorMax(Tmp, Zero);
+		Tmp = VectorMultiply(Tmp, GlobalVectorConstants::BigNumber);
+		Tmp = VectorMin(Tmp, GlobalVectorConstants::FloatOne);
+		*Dst = VectorMax(Tmp, GlobalVectorConstants::FloatZero);
 	}
 };
 
+struct FVectorKernelSelect : public TTrinaryVectorKernel<FVectorKernelSelect>
+{
+	static void VM_FORCEINLINE DoKernel(VectorRegister* RESTRICT Dst, VectorRegister Mask, VectorRegister A, VectorRegister B)
+	{
+		//Currently works by cmpgt 0 to match the current, all vector float vm/scripts but with int support, this should probably change to direct use of a mask.
+		*Dst = VectorSelect(VectorCompareGT(Mask, GlobalVectorConstants::FloatZero), A, B);
+	}
+};
 
 
 struct FVectorKernelSample : public TBinaryVectorKernelData<FVectorKernelSample>
 {
-	static void FORCEINLINE DoKernel(VectorRegister* RESTRICT Dst, const UNiagaraDataObject *Src0, VectorRegister Src1)
+	static void FORCEINLINE DoKernel(VectorRegister* RESTRICT Dst, UNiagaraDataObject *Src0, VectorRegister Src1)
 	{
 		if (Src0)
 		{
@@ -901,12 +1099,234 @@ struct FVectorKernelEaseInOut : public TUnaryVectorKernel<FVectorKernelEaseInOut
 	}
 };
 
+struct FVectorKernelOutputStreamed : public TUnaryVectorKernel<FVectorKernelOutputStreamed>
+{
+	static void VM_FORCEINLINE DoKernel(VectorRegister* Dst, VectorRegister Src0)
+	{
+		VectorStoreAlignedStreamed(Src0, Dst);
+	}
+};
 
 struct FVectorKernelOutput : public TUnaryVectorKernel<FVectorKernelOutput>
 {
 	static void VM_FORCEINLINE DoKernel(VectorRegister* Dst, VectorRegister Src0)
 	{
-		VectorStoreAlignedStreamed(Src0, Dst);
+		VectorStoreAligned(Src0, Dst);
+	}
+};
+
+//////////////////////////////////////////////////////////////////////////
+//Shared data
+
+struct FSharedDataHandlerBase
+{
+	int32 SharedDataIdx;
+	FVectorVMSharedDataView& SharedData;
+	int32 VarIndex;
+	FRegisterHandler<int32, 4> IndexRegisterHandler;
+
+	FSharedDataHandlerBase(FVectorVMContext& Context)
+		: SharedDataIdx(DecodeU8(Context))
+		, SharedData(Context.SharedDataTable[SharedDataIdx])
+		, VarIndex(DecodeU8(Context))
+		, IndexRegisterHandler(Context)
+	{
+	}
+
+	VM_FORCEINLINE void Advance(){ IndexRegisterHandler.Advance(); }
+	VM_FORCEINLINE int32 GetSharedDataIndex(){ return SharedDataIdx; }
+	VM_FORCEINLINE int32 GetVarIndex(){ return VarIndex; }
+	VM_FORCEINLINE EVectorVMOperandLocation GetLocation(){ return EVectorVMOperandLocation::SharedData; }
+	VM_FORCEINLINE int32 GetDataIndex(){ return IndexRegisterHandler.Get(); }
+
+	VM_FORCEINLINE int32 GetLocationIndex(){ return GetSharedDataIndex(); }
+	VM_FORCEINLINE int32 GetSecondaryIndex(){ return GetVarIndex(); }
+	VM_FORCEINLINE int32 GetTertiaryIndex(){ return GetDataIndex(); }
+};
+
+struct FSharedDataHandler : public FSharedDataHandlerBase
+{
+	FSharedDataHandler(FVectorVMContext& Context)
+	: FSharedDataHandlerBase(Context)
+	{
+	}
+
+	VM_FORCEINLINE VectorRegister Get()
+	{
+		return VectorLoad((VectorRegister*)SharedData.GetReadBuffer(VarIndex, IndexRegisterHandler.Get()));
+	}
+};
+
+struct FSharedDataDestHandler : public FSharedDataHandlerBase
+{
+	FSharedDataDestHandler(FVectorVMContext& Context)
+	: FSharedDataHandlerBase(Context)
+	{
+	}
+
+	VM_FORCEINLINE VectorRegister* Get()
+	{
+		return (VectorRegister*)SharedData.GetWriteBuffer(VarIndex, IndexRegisterHandler.Get());
+	}
+	VM_FORCEINLINE VectorRegister GetValue()
+	{
+		return VectorLoadAligned(Get());
+	}
+};
+
+/** 
+Temporary Vector only version of this.
+*/
+struct FVectorKernelSharedDataWrite : public TUnaryKernel<FVectorKernelSharedDataWrite, FSharedDataDestHandler, FRegisterHandler<VectorRegister>>
+{
+	static const VectorVM::EVMType Type = VectorVM::EVMType::Vector4;
+	static const int32 NumInstancesPerOp = 1;
+	static void VM_FORCEINLINE DoKernel(VectorRegister* Buffer, VectorRegister Data)
+	{
+		VectorStoreAlignedStreamed(Data, (float*)(Buffer));
+	}
+}; 
+
+/**
+Temporary Vector only version of this.
+*/
+struct FVectorKernelSharedDataRead : public TUnaryKernel<FVectorKernelSharedDataRead, FRegisterDestHandler<VectorRegister>, FSharedDataHandler>
+{
+	static const VectorVM::EVMType Type = VectorVM::EVMType::Vector4;
+	static const int32 NumInstancesPerOp = 1;
+	static void VM_FORCEINLINE DoKernel(VectorRegister* Dst, VectorRegister Data)
+	{
+		*Dst = Data;
+	}
+};
+
+struct FSharedDataIndexHandlerBase
+{
+	int32 SharedDataIndex;
+	int32 CurrIndex;
+	FVectorVMSharedDataView& SharedData;
+
+	VM_FORCEINLINE void Advance(){}
+	VM_FORCEINLINE EVectorVMOperandLocation GetLocation(){ return EVectorVMOperandLocation::Undefined; }
+	VM_FORCEINLINE int32 GetLocationIndex(){ return SharedDataIndex; }
+	VM_FORCEINLINE int32 GetSecondaryIndex(){ return INDEX_NONE; }
+	VM_FORCEINLINE int32 GetTertiaryIndex(){ return INDEX_NONE; }
+
+	int32 Get(){ return CurrIndex; }
+
+	FSharedDataIndexHandlerBase(FVectorVMContext& Context) 
+		: SharedDataIndex(DecodeU8(Context))
+		, SharedData(Context.SharedDataTable[SharedDataIndex])
+	{	}
+};
+struct FSharedDataIndexHandler_Acquire : public FSharedDataIndexHandlerBase
+{
+	FSharedDataIndexHandler_Acquire(FVectorVMContext& Context) : FSharedDataIndexHandlerBase(Context)	{	}
+	int32 GetNextIndex(){ CurrIndex = SharedData.AcquireIndexWrap(); return CurrIndex; }
+};
+struct FSharedDataIndexHandler_AcquireWrap : public FSharedDataIndexHandlerBase
+{
+	FSharedDataIndexHandler_AcquireWrap(FVectorVMContext& Context) : FSharedDataIndexHandlerBase(Context)	{	}
+	int32 GetNextIndex(){ CurrIndex = SharedData.AcquireIndexWrap();  return CurrIndex; }
+};
+struct FSharedDataIndexHandler_Consume : public FSharedDataIndexHandlerBase
+{
+	FSharedDataIndexHandler_Consume(FVectorVMContext& Context) : FSharedDataIndexHandlerBase(Context)	{	}
+	int32 GetNextIndex(){ CurrIndex = SharedData.ConsumeIndex();  return CurrIndex; }
+};
+struct FSharedDataIndexHandler_ConsumeWrap : public FSharedDataIndexHandlerBase
+{
+	FSharedDataIndexHandler_ConsumeWrap(FVectorVMContext& Context) : FSharedDataIndexHandlerBase(Context)	{	}
+	int32 GetNextIndex(){ CurrIndex = SharedData.ConsumeIndexWrap();  return CurrIndex; }
+};
+
+//Temporary until after scalarization and we can store the size and counter of shared data in constant data.
+template <typename IndexHandler>
+struct FKernelSharedDataGetAppendIndexBase
+{
+	static const VectorVM::EVMType Type = VectorVM::EVMType::Vector4;
+	static const int32 NumInstancesPerOp = 1;
+	static VM_FORCEINLINE void Exec(FVectorVMContext& Context)
+	{
+		FRegisterDestHandler<int32, 4> IndexDest(Context);
+		FRegisterHandler<float, 4> ValidSrc(Context);
+		IndexHandler IdxHandler(Context);
+		int32 NumInstances = Context.NumInstances;
+		Context.BeginOp(VectorVM::EVMType::Vector4, 2, 1);
+		for (int32 i = 0; i < NumInstances; ++i)
+		{
+			Context.PreOp(IndexDest, ValidSrc, IdxHandler);
+
+			int32 Index = ValidSrc.Get() > 0.0f ? IdxHandler.GetNextIndex() : INDEX_NONE;
+			int32* Dest = IndexDest.Get();
+			Dest[0] = Index; Dest[1] = Index; Dest[2] = Index; Dest[3] = Index;
+
+			Context.PostOp(IndexDest, ValidSrc, IdxHandler);
+			ValidSrc.Advance();
+			IndexDest.Advance();
+			IdxHandler.Advance();
+		}
+	}
+};
+
+typedef FKernelSharedDataGetAppendIndexBase<FSharedDataIndexHandler_Acquire> FKernelSharedDataGetAppendIndex;
+typedef FKernelSharedDataGetAppendIndexBase<FSharedDataIndexHandler_AcquireWrap> FKernelSharedDataGetAppendIndex_Wrap;
+
+//Temporary until after scalarization and we can store the size and counter of shared data in constant data.
+template<typename IndexHandler>
+struct FKernelSharedDataGetConsumeIndexBase
+{
+	static const VectorVM::EVMType Type = VectorVM::EVMType::Vector4;
+	static const int32 NumInstancesPerOp = 1;
+
+	static VM_FORCEINLINE void Exec(FVectorVMContext& Context)
+	{
+		FRegisterDestHandler<int32, 4> IndexDest(Context);
+		IndexHandler IdxHandler(Context);
+		int32 NumInstances = Context.NumInstances;
+		Context.BeginOp(VectorVM::EVMType::Vector4, 1, 1);
+		for (int32 i = 0; i < NumInstances; ++i)
+		{
+			Context.PreOp(IndexDest, IdxHandler);
+
+			int32 Index = IdxHandler.GetNextIndex();
+			//Better to just stay in int pipeline?
+			*IndexDest.Get() = Index;//Only need index in X;
+
+			Context.PostOp(IndexDest, IdxHandler);
+			IndexDest.Advance();
+			IdxHandler.Advance();
+		}
+	}
+};
+
+typedef FKernelSharedDataGetConsumeIndexBase<FSharedDataIndexHandler_Consume> FKernelSharedDataGetConsumeIndex;
+typedef FKernelSharedDataGetConsumeIndexBase<FSharedDataIndexHandler_ConsumeWrap> FKernelSharedDataGetConsumeIndex_Wrap;
+
+//Temporary until after scalarization and we can store the size and counter of shared data in constant data.
+struct FKernelSharedDataIndexValid
+{
+	static VM_FORCEINLINE void Exec(FVectorVMContext& Context)
+	{
+		FRegisterDestHandler<float, 4> ValidDest(Context);
+		FRegisterHandler<int32, 4> IndexSrc(Context);
+		FVectorVMSharedDataView& SharedData = Context.SharedDataTable[DecodeU8(Context)];
+		int32 NumInstances = Context.NumInstances;
+		Context.BeginOp(VectorVM::EVMType::Vector4, 1, 1);
+		for (int32 i = 0; i < NumInstances; ++i)
+		{
+			Context.PreOp(ValidDest, IndexSrc);
+
+			int32 Index = IndexSrc.Get();
+			float Valid = SharedData.ValidIndex(Index) ? 1.0f : 0.0f;
+			//Better to just stay in int pipeline?
+			float* Dst = ValidDest.Get();
+			Dst[0] = Valid; Dst[1] = Valid; Dst[2] = Valid; Dst[3] = Valid;
+
+			Context.PostOp(ValidDest, IndexSrc);
+			ValidDest.Advance();
+			IndexSrc.Advance();;
+		}
 	}
 };
 
@@ -990,6 +1410,27 @@ void VectorVM::Init()
 			}
 		}
 
+#if WITH_EDITOR
+
+		OpNames.AddDefaulted((int32)EVectorVMOp::NumOpcodes);
+		if (UEnum* EnumStateObj = FindObject<UEnum>(ANY_PACKAGE, TEXT("EVectorVMOp"), true))
+		{
+			for (int32 i = 0; i < (int32)EVectorVMOp::NumOpcodes; ++i)
+			{
+				OpNames[i] = EnumStateObj->GetDisplayNameText(i).ToString();
+			}
+		}
+
+		OperandLocationNames.AddDefaulted((int32)EVectorVMOperandLocation::Num);
+		if (UEnum* EnumStateObj = FindObject<UEnum>(ANY_PACKAGE, TEXT("EVectorVMOperandLocation"), true))
+		{
+			for (int32 i = 0; i < (int32)EVectorVMOperandLocation::Num; ++i)
+			{
+				OperandLocationNames[i] = EnumStateObj->GetDisplayNameText(i).ToString();
+			}
+		}
+
+#endif
 
 		Inited = true;
 	}
@@ -1003,6 +1444,7 @@ void VectorVM::Exec(
 	int32 NumOutputRegisters,
 	FVector4 const* ConstantTable,
 	UNiagaraDataObject* *DataObjConstTable,
+	FVectorVMSharedDataView* SharedDataTable,
 	int32 NumVectors
 	)
 {
@@ -1031,74 +1473,94 @@ void VectorVM::Exec(
 
 		// Setup execution context.
 		int32 VectorsThisChunk = FMath::Min<int32>(NumVectors, VectorsPerChunk);
-		FVectorVMContext Context(Code, RegisterTable, ConstantTable, DataObjConstTable, VectorsThisChunk);
-		EOp Op = EOp::done;
+		int32 InstancesThisChunk = VectorsThisChunk;
+		FVectorVMContext Context(Code, RegisterTable, ConstantTable, DataObjConstTable, SharedDataTable, VectorsThisChunk, InstancesThisChunk, VectorsPerChunk * ChunkIndex);
+		EVectorVMOp Op = EVectorVMOp::done;
 
 		// Execute VM on all vectors in this chunk.
 		do 
 		{
 			Op = DecodeOp(Context);
+			Context.SetOp(Op);
 			switch (Op)
 			{
 			// Dispatch kernel ops.
-			case EOp::add: FVectorKernelAdd::Exec(Context); break;
-			case EOp::sub: FVectorKernelSub::Exec(Context); break;
-			case EOp::mul: FVectorKernelMul::Exec(Context); break;
-			case EOp::div: FVectorKernelDiv::Exec(Context); break;
-			case EOp::mad: FVectorKernelMad::Exec(Context); break;
-			case EOp::lerp: FVectorKernelLerp::Exec(Context); break;
-			case EOp::rcp: FVectorKernelRcp::Exec(Context); break;
-			case EOp::rsq: FVectorKernelRsq::Exec(Context); break;
-			case EOp::sqrt: FVectorKernelSqrt::Exec(Context); break;
-			case EOp::neg: FVectorKernelNeg::Exec(Context); break;
-			case EOp::abs: FVectorKernelAbs::Exec(Context); break;
-			case EOp::exp: FVectorKernelExp::Exec(Context); break;
-			case EOp::exp2: FVectorKernelExp2::Exec(Context); break;
-			case EOp::log: FVectorKernelLog::Exec(Context); break;
-			case EOp::log2: FVectorKernelLog2::Exec(Context); break;
-			case EOp::sin: FVectorKernelSin::Exec(Context); break;
-			case EOp::cos: FVectorKernelCos::Exec(Context); break;
-			case EOp::tan: FVectorKernelTan::Exec(Context); break;
-			case EOp::asin: FVectorKernelASin::Exec(Context); break;
-			case EOp::acos: FVectorKernelACos::Exec(Context); break;
-			case EOp::atan: FVectorKernelATan::Exec(Context); break;
-			case EOp::atan2: FVectorKernelATan2::Exec(Context); break;
-			case EOp::ceil: FVectorKernelCeil::Exec(Context); break;
-			case EOp::floor: FVectorKernelFloor::Exec(Context); break;
-			case EOp::fmod: FVectorKernelMod::Exec(Context); break;
-			case EOp::frac: FVectorKernelFrac::Exec(Context); break;
-			case EOp::trunc: FVectorKernelTrunc::Exec(Context); break;
-			case EOp::clamp: FVectorKernelClamp::Exec(Context); break;
-			case EOp::min: FVectorKernelMin::Exec(Context); break;
-			case EOp::max: FVectorKernelMax::Exec(Context); break;
-			case EOp::pow: FVectorKernelPow::Exec(Context); break;
-			case EOp::sign: FVectorKernelSign::Exec(Context); break;
-			case EOp::step: FVectorKernelStep::Exec(Context); break;
-			case EOp::dot: FVectorKernelDot::Exec(Context); break;
-			case EOp::cross: FVectorKernelCross::Exec(Context); break;
-			case EOp::normalize: FVectorKernelNormalize::Exec(Context); break;
-			case EOp::random: FVectorKernelRandom::Exec(Context); break;
-			case EOp::length: FVectorKernelLength::Exec(Context); break;
-			case EOp::noise: FVectorKernelNoise::Exec(Context); break;
-			case EOp::splatx: FVectorKernelSplat<0>::Exec(Context); break;
-			case EOp::splaty: FVectorKernelSplat<1>::Exec(Context); break;
-			case EOp::splatz: FVectorKernelSplat<2>::Exec(Context); break;
-			case EOp::splatw: FVectorKernelSplat<3>::Exec(Context); break;
-			case EOp::compose: FVectorKernelCompose<0,1,2,3>::Exec(Context); break;
-			case EOp::composex: FVectorKernelCompose<0, 0, 0, 0>::Exec(Context); break;
-			case EOp::composey: FVectorKernelCompose<1, 1, 1, 1>::Exec(Context); break;
-			case EOp::composez: FVectorKernelCompose<2, 2, 2, 2>::Exec(Context); break;
-			case EOp::composew: FVectorKernelCompose<3, 3, 3, 3>::Exec(Context); break;
-			case EOp::lessthan: FVectorKernelLessThan::Exec(Context); break;
-			case EOp::sample: FVectorKernelSample::Exec(Context); break;
-			case EOp::bufferwrite: FVectorKernelBufferWrite::Exec(Context); break;
-			case EOp::eventbroadcast: break;
-			case EOp::easein: FVectorKernelEaseIn::Exec(Context); break;
-			case EOp::easeinout: FVectorKernelEaseInOut::Exec(Context); break;
-			case EOp::output: FVectorKernelOutput::Exec(Context); break;
-				
+			case EVectorVMOp::add: FVectorKernelAdd::Exec(Context); break;
+			case EVectorVMOp::sub: FVectorKernelSub::Exec(Context); break;
+			case EVectorVMOp::mul: FVectorKernelMul::Exec(Context); break;
+			case EVectorVMOp::div: FVectorKernelDiv::Exec(Context); break;
+			case EVectorVMOp::mad: FVectorKernelMad::Exec(Context); break;
+			case EVectorVMOp::lerp: FVectorKernelLerp::Exec(Context); break;
+			case EVectorVMOp::rcp: FVectorKernelRcp::Exec(Context); break;
+			case EVectorVMOp::rsq: FVectorKernelRsq::Exec(Context); break;
+			case EVectorVMOp::sqrt: FVectorKernelSqrt::Exec(Context); break;
+			case EVectorVMOp::neg: FVectorKernelNeg::Exec(Context); break;
+			case EVectorVMOp::abs: FVectorKernelAbs::Exec(Context); break;
+			case EVectorVMOp::exp: FVectorKernelExp::Exec(Context); break;
+			case EVectorVMOp::exp2: FVectorKernelExp2::Exec(Context); break;
+			case EVectorVMOp::log: FVectorKernelLog::Exec(Context); break;
+			case EVectorVMOp::log2: FVectorKernelLog2::Exec(Context); break;
+			case EVectorVMOp::sin: FVectorKernelSin::Exec(Context); break;
+			case EVectorVMOp::cos: FVectorKernelCos::Exec(Context); break;
+			case EVectorVMOp::tan: FVectorKernelTan::Exec(Context); break;
+			case EVectorVMOp::asin: FVectorKernelASin::Exec(Context); break;
+			case EVectorVMOp::acos: FVectorKernelACos::Exec(Context); break;
+			case EVectorVMOp::atan: FVectorKernelATan::Exec(Context); break;
+			case EVectorVMOp::atan2: FVectorKernelATan2::Exec(Context); break;
+			case EVectorVMOp::ceil: FVectorKernelCeil::Exec(Context); break;
+			case EVectorVMOp::floor: FVectorKernelFloor::Exec(Context); break;
+			case EVectorVMOp::fmod: FVectorKernelMod::Exec(Context); break;
+			case EVectorVMOp::frac: FVectorKernelFrac::Exec(Context); break;
+			case EVectorVMOp::trunc: FVectorKernelTrunc::Exec(Context); break;
+			case EVectorVMOp::clamp: FVectorKernelClamp::Exec(Context); break;
+			case EVectorVMOp::min: FVectorKernelMin::Exec(Context); break;
+			case EVectorVMOp::max: FVectorKernelMax::Exec(Context); break;
+			case EVectorVMOp::pow: FVectorKernelPow::Exec(Context); break;
+			case EVectorVMOp::sign: FVectorKernelSign::Exec(Context); break;
+			case EVectorVMOp::step: FVectorKernelStep::Exec(Context); break;
+			case EVectorVMOp::dot: FVectorKernelDot::Exec(Context); break;
+			case EVectorVMOp::cross: FVectorKernelCross::Exec(Context); break;
+			case EVectorVMOp::normalize: FVectorKernelNormalize::Exec(Context); break;
+			case EVectorVMOp::random: FVectorKernelRandom::Exec(Context); break;
+			case EVectorVMOp::length: FVectorKernelLength::Exec(Context); break;
+			case EVectorVMOp::noise: FVectorKernelNoise::Exec(Context); break;
+			case EVectorVMOp::splatx: FVectorKernelSplat<0>::Exec(Context); break;
+			case EVectorVMOp::splaty: FVectorKernelSplat<1>::Exec(Context); break;
+			case EVectorVMOp::splatz: FVectorKernelSplat<2>::Exec(Context); break;
+			case EVectorVMOp::splatw: FVectorKernelSplat<3>::Exec(Context); break;
+			case EVectorVMOp::compose: FVectorKernelCompose<0,1,2,3>::Exec(Context); break;
+			case EVectorVMOp::composex: FVectorKernelCompose<0, 0, 0, 0>::Exec(Context); break;
+			case EVectorVMOp::composey: FVectorKernelCompose<1, 1, 1, 1>::Exec(Context); break;
+			case EVectorVMOp::composez: FVectorKernelCompose<2, 2, 2, 2>::Exec(Context); break;
+			case EVectorVMOp::composew: FVectorKernelCompose<3, 3, 3, 3>::Exec(Context); break;
+			case EVectorVMOp::lessthan: FVectorKernelLessThan::Exec(Context); break;
+			case EVectorVMOp::sample: FVectorKernelSample::Exec(Context); break;
+			case EVectorVMOp::bufferwrite: FVectorKernelBufferWrite::Exec(Context); break;
+			case EVectorVMOp::easein: FVectorKernelEaseIn::Exec(Context); break;
+			case EVectorVMOp::easeinout: FVectorKernelEaseInOut::Exec(Context); break;
+			case EVectorVMOp::aquireshareddataindex: FKernelSharedDataGetAppendIndex::Exec(Context); break;
+			case EVectorVMOp::aquireshareddataindexwrap: FKernelSharedDataGetAppendIndex_Wrap::Exec(Context); break;
+			case EVectorVMOp::consumeshareddataindex: FKernelSharedDataGetConsumeIndex::Exec(Context); break;
+			case EVectorVMOp::consumeshareddataindexwrap: FKernelSharedDataGetConsumeIndex_Wrap::Exec(Context); break;
+			case EVectorVMOp::shareddataread: FVectorKernelSharedDataRead::Exec(Context); break;
+			case EVectorVMOp::shareddatawrite: FVectorKernelSharedDataWrite::Exec(Context); break;
+			case EVectorVMOp::shareddataindexvalid: FKernelSharedDataIndexValid::Exec(Context); break;
+			case EVectorVMOp::select: FVectorKernelSelect::Exec(Context); break;
+			case EVectorVMOp::output:
+			{
+				if ((ENABLE_VM_DEBUGGING && Context.IsDebugging()) == false)
+				{
+					FVectorKernelOutput::Exec(Context);
+				}
+				else
+				{
+					FVectorKernelOutputStreamed::Exec(Context);
+				}
+			}
+				break;
+
 			// Execution always terminates with a "done" opcode.
-			case EOp::done:
+			case EVectorVMOp::done:
 				break;
 
 			// Opcode not recognized / implemented.
@@ -1106,7 +1568,7 @@ void VectorVM::Exec(
 				UE_LOG(LogVectorVM, Error, TEXT("Unknown op code 0x%02x"), (uint32)Op);
 				return;//BAIL
 			}
-		} while (Op != EOp::done);
+		} while (Op != EVectorVMOp::done);
 
 		NumVectors -= VectorsPerChunk;
 	}
@@ -1114,9 +1576,21 @@ void VectorVM::Exec(
 
 uint8 VectorVM::GetNumOpCodes()
 {
-	return (uint8)EOp::NumOpcodes;
+	return (uint8)EVectorVMOp::NumOpcodes;
+}
+
+#if WITH_EDITOR
+FString VectorVM::GetOpName(EVectorVMOp Op)
+{
+	return OpNames[(int32)Op];
 }
 
+FString VectorVM::GetOperandLocationName(EVectorVMOperandLocation Location)
+{
+	return OperandLocationNames[(int32)Location];
+}
+#endif
+
 /*------------------------------------------------------------------------------
 	Automation test for the VM.
 ------------------------------------------------------------------------------*/
@@ -1125,14 +1599,16 @@ IMPLEMENT_SIMPLE_AUTOMATION_TEST(FVectorVMTest, "System.Core.Math.Vector VM", EA
 
 bool FVectorVMTest::RunTest(const FString& Parameters)
 {
+	VectorVM::Init();
+
 	uint8 TestCode[] =
 	{
-		(uint8)VectorVM::EOp::mul, 0x00, SRCOP_RRRR, 0x0 + VectorVM::NumTempRegisters, 0x0 + VectorVM::NumTempRegisters,       // mul r0, r8, r8
-		(uint8)VectorVM::EOp::mad, 0x01, SRCOP_RRRR, 0x01 + VectorVM::NumTempRegisters, 0x01 + VectorVM::NumTempRegisters, 0x00, // mad r1, r9, r9, r0
-		(uint8)VectorVM::EOp::mad, 0x00, SRCOP_RRRR, 0x02 + VectorVM::NumTempRegisters, 0x02 + VectorVM::NumTempRegisters, 0x01, // mad r0, r10, r10, r1
-		(uint8)VectorVM::EOp::add, 0x01, SRCOP_RRCR, 0x00, 0x01,       // addi r1, r0, c1
-		(uint8)VectorVM::EOp::neg, 0x00, SRCOP_RRRR, 0x01,             // neg r0, r1
-		(uint8)VectorVM::EOp::clamp, VectorVM::FirstOutputRegister, SRCOP_RCCR, 0x00, 0x02, 0x03, // clampii r40, r0, c2, c3
+		(uint8)EVectorVMOp::mul, SRCOP_RRRR, 0x00, 0x0 + VectorVM::NumTempRegisters, 0x0 + VectorVM::NumTempRegisters,       // mul r0, r8, r8
+		(uint8)EVectorVMOp::mad, SRCOP_RRRR, 0x01, 0x01 + VectorVM::NumTempRegisters, 0x01 + VectorVM::NumTempRegisters, 0x00, // mad r1, r9, r9, r0
+		(uint8)EVectorVMOp::mad, SRCOP_RRRR, 0x00, 0x02 + VectorVM::NumTempRegisters, 0x02 + VectorVM::NumTempRegisters, 0x01, // mad r0, r10, r10, r1
+		(uint8)EVectorVMOp::add, SRCOP_RRCR, 0x01, 0x00, 0x01,       // addi r1, r0, c1
+		(uint8)EVectorVMOp::neg, SRCOP_RRRR, 0x00, 0x01,             // neg r0, r1
+		(uint8)EVectorVMOp::clamp, SRCOP_RCCR, VectorVM::FirstOutputRegister, 0x00, 0x02, 0x03, // clampii r40, r0, c2, c3
 		0x00 // terminator
 	};
 
@@ -1157,12 +1633,15 @@ bool FVectorVMTest::RunTest(const FString& Parameters)
 	ConstantTable[2] = FVector4(-20.0f, -20.0f, -20.0f, -20.0f);
 	ConstantTable[3] = FVector4(20.0f, 20.0f, 20.0f, 20.0f);
 
+	FVectorVMSharedDataView* dummy = NULL;
+
 	VectorVM::Exec(
 		TestCode,
 		InputRegisters, 3,
 		OutputRegisters, 1,
 		ConstantTable,
 		nullptr,
+		dummy,
 		VectorVM::VectorsPerChunk
 		);
 
@@ -1207,4 +1686,136 @@ bool FVectorVMTest::RunTest(const FString& Parameters)
 	return true;
 }
 
+//////////////////////////////////////////////////////////////////////////
+#if ENABLE_VM_DEBUGGING
+
+void FVectorVMDebuggerImpl::BeginOp(FVectorVMContext& Context, VectorVM::EVMType InType, int32 InNumArgs, int32 InNumInstancesPerOp)
+{
+	CurrOp = Context.CurrOp;
+	OpType = InType;
+	CurrNumArgs = InNumArgs;
+	NumInstancesPerOp = InNumInstancesPerOp;
+	CurrInstanceBase = StartInstance + Context.StartInstance;
+}
+
+template<typename DstHandler, typename Arg0Handler, typename Arg1Handler, typename Arg2Handler, typename Arg3Handler>
+void FVectorVMDebuggerImpl::PreOp(FVectorVMContext& Context, DstHandler& Dst, Arg0Handler& Arg0, Arg1Handler& Arg1, Arg2Handler& Arg2, Arg3Handler& Arg3)
+{
+	CachedPreOpData[(int32)VectorVM::EKernelArgs::Dest].Set(Dst.GetValue());
+	CachedPreOpData[(int32)VectorVM::EKernelArgs::Arg0].Set(Arg0.Get());
+	CachedPreOpData[(int32)VectorVM::EKernelArgs::Arg1].Set(Arg1.Get());
+	CachedPreOpData[(int32)VectorVM::EKernelArgs::Arg2].Set(Arg2.Get());
+	CachedPreOpData[(int32)VectorVM::EKernelArgs::Arg3].Set(Arg3.Get());
+}
+
+template<typename DstHandler, typename Arg0Handler, typename Arg1Handler, typename Arg2Handler, typename Arg3Handler >
+void FVectorVMDebuggerImpl::PostOp(FVectorVMContext& Context, DstHandler& Dst, Arg0Handler& Arg0, Arg1Handler& Arg1, Arg2Handler& Arg2, Arg3Handler& Arg3)
+{
+	for (int32 i = 0; i < NumInstancesPerOp; ++i)
+	{
+		if (TArray<VectorVM::FOpDebugInfo>* Info = DebugInfo.Find(CurrInstanceBase + i))
+		{
+			TArray<VectorVM::FOpDebugInfo>& InfoArray = *Info;
+			int32 NewIdx = InfoArray.AddUninitialized();
+			VectorVM::FOpDebugInfo& NewInfo = InfoArray[NewIdx];
+			NewInfo.NumArgs = CurrNumArgs;
+			NewInfo.Op = CurrOp;
+			NewInfo.OpType = OpType;
+			NewInfo.LocationInfo[(int32)VectorVM::EKernelArgs::Dest] = VectorVM::FDataLocationInfo(Context, Dst.GetLocation(), Dst.GetLocationIndex(), Dst.GetSecondaryIndex(), Dst.GetTertiaryIndex());
+			NewInfo.LocationInfo[(int32)VectorVM::EKernelArgs::Arg0] = VectorVM::FDataLocationInfo(Context, Arg0.GetLocation(), Arg0.GetLocationIndex(), Arg0.GetSecondaryIndex(), Arg0.GetTertiaryIndex());;
+			NewInfo.LocationInfo[(int32)VectorVM::EKernelArgs::Arg1] = VectorVM::FDataLocationInfo(Context, Arg1.GetLocation(), Arg1.GetLocationIndex(), Arg1.GetSecondaryIndex(), Arg1.GetTertiaryIndex());;
+			NewInfo.LocationInfo[(int32)VectorVM::EKernelArgs::Arg2] = VectorVM::FDataLocationInfo(Context, Arg2.GetLocation(), Arg2.GetLocationIndex(), Arg2.GetSecondaryIndex(), Arg2.GetTertiaryIndex());;
+			NewInfo.LocationInfo[(int32)VectorVM::EKernelArgs::Arg3] = VectorVM::FDataLocationInfo(Context, Arg3.GetLocation(), Arg3.GetLocationIndex(), Arg3.GetSecondaryIndex(), Arg3.GetTertiaryIndex());;
+
+			NewInfo.PreOpValues[(int32)VectorVM::EKernelArgs::Dest] = CachedPreOpData[(int32)VectorVM::EKernelArgs::Dest];
+			NewInfo.PreOpValues[(int32)VectorVM::EKernelArgs::Arg0] = CachedPreOpData[(int32)VectorVM::EKernelArgs::Arg0];
+			NewInfo.PreOpValues[(int32)VectorVM::EKernelArgs::Arg1] = CachedPreOpData[(int32)VectorVM::EKernelArgs::Arg1];
+			NewInfo.PreOpValues[(int32)VectorVM::EKernelArgs::Arg2] = CachedPreOpData[(int32)VectorVM::EKernelArgs::Arg2];
+			NewInfo.PreOpValues[(int32)VectorVM::EKernelArgs::Arg3] = CachedPreOpData[(int32)VectorVM::EKernelArgs::Arg3];
+
+			NewInfo.PostOpValues[(int32)VectorVM::EKernelArgs::Dest].Set(Dst.GetValue());
+			NewInfo.PostOpValues[(int32)VectorVM::EKernelArgs::Arg0].Set(Arg0.Get());
+			NewInfo.PostOpValues[(int32)VectorVM::EKernelArgs::Arg1].Set(Arg1.Get());
+			NewInfo.PostOpValues[(int32)VectorVM::EKernelArgs::Arg2].Set(Arg2.Get());
+			NewInfo.PostOpValues[(int32)VectorVM::EKernelArgs::Arg3].Set(Arg3.Get());
+		}
+	}
+	CurrInstanceBase += NumInstancesPerOp;
+}
+
+void VectorVM::AttachDebugger(VectorVM::FVectorVMDebugger* Debugger)
+{
+	check(IsInGameThread());
+	check(Debugger);
+	check(AttachedDebugger == NULL);
+	AttachedDebugger = Debugger->GetImpl();
+}
+
+void VectorVM::DetachDebugger(VectorVM::FVectorVMDebugger* Debugger)
+{
+	check(IsInGameThread());
+	check(Debugger);
+	check(AttachedDebugger == Debugger->GetImpl());
+	AttachedDebugger = NULL;
+}
+
+VectorVM::FVectorVMDebugger::FVectorVMDebugger()
+: Impl(new FVectorVMDebuggerImpl())
+{
+
+}
+
+VectorVM::FVectorVMDebugger::~FVectorVMDebugger()
+{
+	if (Impl)
+	{
+		delete Impl;
+	}
+}
+
+void VectorVM::FVectorVMDebugger::AddInstanceToGather(int32 Instance)
+{
+	check(Impl);
+	Impl->DebugInfo.Add(Instance);
+}
+
+const TArray<VectorVM::FOpDebugInfo>* VectorVM::FVectorVMDebugger::GetDebugInfoForInstance(int32 Instance)
+{
+	check(Impl);
+	return Impl->DebugInfo.Find(Instance);
+}
+
+FVectorVMDebuggerImpl* VectorVM::FVectorVMDebugger::GetImpl()
+{
+	return Impl;
+}
+
+void VectorVM::FVectorVMDebugger::InitForScriptRun(int32 StartInstance, const TArray<int32> InstancesToDebug)
+{
+	Impl->InitForScriptRun(StartInstance, InstancesToDebug);
+}
+
+VectorVM::FDataLocationInfo::FDataLocationInfo(FVectorVMContext& Context, EVectorVMOperandLocation InLocation, int32 InPrimaryIndex, int32 InSecondaryIndex, int32 InTertiaryIndex)
+: Location(InLocation)
+, PrimaryLocationIndex(InPrimaryIndex)
+, SecondaryLocationIndex(InSecondaryIndex)
+, TertiaryLocationIndex(InTertiaryIndex)
+{
+	switch (Location)
+	{
+	case EVectorVMOperandLocation::Constant: MemoryAddress = &Context.ConstantTable[PrimaryLocationIndex]; break;
+	case EVectorVMOperandLocation::TemporaryRegister: MemoryAddress = &Context.RegisterTable[PrimaryLocationIndex]; break;
+	case EVectorVMOperandLocation::InputRegister: MemoryAddress = &Context.RegisterTable[NumTempRegisters + PrimaryLocationIndex]; break;
+	case EVectorVMOperandLocation::OutputRegister: MemoryAddress = &Context.ConstantTable[NumTempRegisters + MaxInputRegisters + PrimaryLocationIndex]; break;
+	case EVectorVMOperandLocation::SharedData:
+	{
+												 FVectorVMSharedDataView& SharedData = Context.SharedDataTable[PrimaryLocationIndex];
+												 MemoryAddress = SharedData.GetReadBuffer(SecondaryLocationIndex, TertiaryLocationIndex);
+	} break;
+	default: MemoryAddress = nullptr; break;
+	};
+}
+
+#endif 
+
 #undef VM_FORCEINLINE
\ No newline at end of file
diff --git a/Engine/Source/Runtime/VectorVM/Private/VectorVMPrivate.h b/Engine/Source/Runtime/VectorVM/Private/VectorVMPrivate.h
--- a/Engine/Source/Runtime/VectorVM/Private/VectorVMPrivate.h
+++ b/Engine/Source/Runtime/VectorVM/Private/VectorVMPrivate.h
@@ -15,3 +15,69 @@ namespace VectorVM
 		VectorsPerChunk = ChunkSize / ElementsPerVector,
 	};
 }
+
+struct FDummyHandler
+{
+	FORCEINLINE void Advance(){ }
+	FORCEINLINE int32 GetLocationIndex(){ return INDEX_NONE; }
+	FORCEINLINE int32 GetSecondaryIndex(){ return INDEX_NONE; }
+	FORCEINLINE int32 GetTertiaryIndex(){ return INDEX_NONE; }
+	FORCEINLINE EVectorVMOperandLocation GetLocation(){ return EVectorVMOperandLocation::Undefined; }
+	FORCEINLINE VectorRegister Get(){ return GlobalVectorConstants::FloatZero; }
+	FORCEINLINE VectorRegister GetValue(){ return GlobalVectorConstants::FloatZero; }
+};
+extern FDummyHandler DummyHandler;
+
+//////////////////////////////////////////////////////////////////////////
+// Debugger
+#if ENABLE_VM_DEBUGGING
+
+class FVectorVMDebuggerImpl
+{
+	friend class VectorVM::FVectorVMDebugger;
+
+	FVectorVMDebuggerImpl()
+		: CurrOp(EVectorVMOp::done)
+		, OpType(VectorVM::EVMType::Vector4)
+		, CurrNumArgs(0)
+		, CurrInstanceBase(0)
+		, NumInstancesPerOp(0)
+		, StartInstance(0)
+	{
+
+	}
+	//Map of instance index to debug info gathered for that instance.
+	TMap<int32, TArray<VectorVM::FOpDebugInfo>> DebugInfo;
+
+	EVectorVMOp CurrOp;
+	VectorVM::EVMType OpType;
+	int32 CurrNumArgs;
+	int32 CurrInstanceBase;
+	int32 NumInstancesPerOp;
+	int32 StartInstance;
+
+	VectorVM::FDebugValue CachedPreOpData[NUM_VM_OP_DEBUG_VALUES];
+
+public:
+
+	void InitForScriptRun(int32 InStartInstance, const TArray<int32> InstancesToDebug)
+	{
+		StartInstance = InStartInstance;
+		DebugInfo.Empty(InstancesToDebug.Num());
+		for (int32 i : InstancesToDebug)
+		{
+			DebugInfo.Add(i);
+		}
+	}
+
+	void BeginOp(FVectorVMContext& Context, VectorVM::EVMType InType, int32 InNumArgs, int32 InNumInstancesPerOp);
+
+	template<typename DstHandler, typename Arg0Handler, typename Arg1Handler = FDummyHandler, typename Arg2Handler = FDummyHandler, typename Arg3Handler = FDummyHandler>
+	void PreOp(FVectorVMContext& Context, DstHandler& Dst, Arg0Handler& Arg, Arg1Handler& Arg1 = DummyHandler, Arg2Handler& Arg2 = DummyHandler, Arg3Handler& Arg3 = DummyHandler);
+
+	template<typename DstHandler, typename Arg0Handler, typename Arg1Handler = FDummyHandler, typename Arg2Handler = FDummyHandler, typename Arg3Handler = FDummyHandler>
+	void PostOp(FVectorVMContext& Context, DstHandler& Dst, Arg0Handler& Arg, Arg1Handler& Arg1 = DummyHandler, Arg2Handler& Arg2 = DummyHandler, Arg3Handler& Arg3 = DummyHandler);
+};
+
+#endif
+
diff --git a/Engine/Source/Runtime/VectorVM/Public/VectorVM.h b/Engine/Source/Runtime/VectorVM/Public/VectorVM.h
--- a/Engine/Source/Runtime/VectorVM/Public/VectorVM.h
+++ b/Engine/Source/Runtime/VectorVM/Public/VectorVM.h
@@ -4,7 +4,97 @@
 #include "Core.h"
 #include "CoreUObject.h"
 
+//TODO: move to a per platform header and have VM scale vectorization according to vector width.
+#define VECTOR_WIDTH (128)
+#define VECTOR_WIDTH_BYTES (16)
+#define VECTOR_WIDTH_FLOATS (4)
+
 class UNiagaraDataObject;
+class FVectorVMSharedDataView;
+struct FVectorVMContext;
+class FVectorVMDebuggerImpl;
+
+UENUM()
+enum class EVectorVMOperandLocation : uint8
+{
+	TemporaryRegister,
+	InputRegister,
+	OutputRegister,
+	Constant,
+	DataObjConstant,
+	SharedData,
+	Undefined,
+	Num
+};
+
+UENUM()
+enum class EVectorVMOp : uint8
+{
+	done,
+	add,
+	sub,
+	mul,
+	mad,
+	lerp,
+	rcp,
+	rsq,
+	sqrt,
+	neg,
+	abs,
+	exp,
+	exp2,
+	log,
+	log2,
+	sin,
+	cos,
+	tan,
+	asin,
+	acos,
+	atan,
+	atan2,
+	ceil,
+	floor,
+	fmod,
+	frac,
+	trunc,
+	clamp,
+	min,
+	max,
+	pow,
+	sign,
+	step,
+	dot,
+	cross,
+	normalize,
+	random,
+	length,
+	noise,
+	splatx,
+	splaty,
+	splatz,
+	splatw,
+	compose,
+	composex,
+	composey,
+	composez,
+	composew,
+	output,
+	lessthan,
+	select,
+	sample,
+	bufferwrite,
+	easein,
+	easeinout,
+	div,
+	aquireshareddataindex,
+	aquireshareddataindexwrap,
+	consumeshareddataindex,
+	consumeshareddataindexwrap,
+	shareddataread,
+	shareddatawrite,
+	shareddataindexvalid,
+	NumOpcodes
+};
 
 namespace VectorVM
 {
@@ -21,81 +111,15 @@ namespace VectorVM
 		MaxRegisters = NumTempRegisters + MaxInputRegisters + MaxOutputRegisters + MaxConstants,
 	};
 
-	enum EOperandType
-	{
-		UndefinedOperandType = 0,
-		RegisterOperandType,
-		ConstantOperandType,
-		DataObjConstantOperandType,
-		TemporaryOperandType
-	};
-
-	/** List of opcodes supported by the VM. */
-	enum class EOp
-	{
-			done,
-			add,
-			sub,
-			mul,
-			mad,
-			lerp,
-			rcp,
-			rsq,
-			sqrt,
-			neg,
-			abs,
-			exp,
-			exp2,
-			log,
-			log2,
-			sin,
-			cos,
-			tan,
-			asin,
-			acos,
-			atan,
-			atan2,
-			ceil,
-			floor,
-			fmod,
-			frac,
-			trunc,
-			clamp,
-			min,
-			max,
-			pow,
-			sign,
-			step,
-			dot,
-			cross,
-			normalize,
-			random,
-			length,
-			noise,
-			splatx,
-			splaty,
-			splatz,
-			splatw,
-			compose,
-			composex,
-			composey,
-			composez,
-			composew,
-			output,
-			lessthan,
-			sample,
-			bufferwrite,
-			eventbroadcast,
-			easein,
-			easeinout,
-			div,
-			NumOpcodes
-	};
-
 	/** Get total number of op-codes */
 	VECTORVM_API uint8 GetNumOpCodes();
 
-	VECTORVM_API uint8 CreateSrcOperandMask(VectorVM::EOperandType Type1, VectorVM::EOperandType Type2 = VectorVM::RegisterOperandType, VectorVM::EOperandType Type3 = VectorVM::RegisterOperandType, VectorVM::EOperandType Type4 = VectorVM::RegisterOperandType);
+#if WITH_EDITOR
+	VECTORVM_API FString GetOpName(EVectorVMOp Op);
+	VECTORVM_API FString GetOperandLocationName(EVectorVMOperandLocation Location);
+#endif
+
+	VECTORVM_API uint8 CreateSrcOperandMask(EVectorVMOperandLocation Type1, EVectorVMOperandLocation Type2 = EVectorVMOperandLocation::TemporaryRegister, EVectorVMOperandLocation Type3 = EVectorVMOperandLocation::TemporaryRegister, EVectorVMOperandLocation Type4 = EVectorVMOperandLocation::TemporaryRegister);
 
 	/**
 	 * Execute VectorVM bytecode.
@@ -108,12 +132,171 @@ namespace VectorVM
 		int32 NumOutputRegisters,
 		FVector4 const* ConstantTable,
 		UNiagaraDataObject* *DataObjTable,
+		FVectorVMSharedDataView* SharedDataTable,
 		int32 NumVectors
 		);
 
 	VECTORVM_API void Init();
 
+	//////////////////////////////////////////////////////////////////////////
+	//Debugging
+
+#define ENABLE_VM_DEBUGGING (0)
+//#define ENABLE_VM_DEBUGGING (1)
+
+#define MAX_VM_OP_ARGS (4)
+#define NUM_VM_OP_DEBUG_VALUES (MAX_VM_OP_ARGS + 1)
+#define MAX_INSTANCES_PER_VM_OP (4)//8 and 16 in future
+
+	enum class EVMType : uint8
+	{
+		//Float,
+		//Int32,
+		Vector4,
+	};
+
+	enum class EKernelArgs : uint8
+	{
+		Dest,
+		Arg0,
+		Arg1,
+		Arg2,
+		Arg3,
+		Max,
+	};
+
+#if ENABLE_VM_DEBUGGING
+
+	struct FDebugValue
+	{
+		FVector4 v;//TODO, replace with int/float union?
+
+		void Set(float Val)
+		{
+			v = FVector4(Val, Val, Val, Val);
+		}
+		void Set(int32 Val)
+		{
+			int32* i = (int32*)&v;
+			i[0] = Val;	i[1] = Val; i[2] = Val; i[3] = Val;
+		}
+		void Set(VectorRegister Val)
+		{
+			VectorStoreAligned(Val, (float*)&v);
+		}
+		template<typename T>
+		void Set(T Val)
+		{
+			v = FVector4(0.0f, 0.0f, 0.0f, 0.0f);
+		}
+	};
+
+	struct FDataLocationInfo
+	{
+		EVectorVMOperandLocation Location;
+		int32 PrimaryLocationIndex;
+		int32 SecondaryLocationIndex;
+		int32 TertiaryLocationIndex;
+		const void* MemoryAddress;
+
+		FDataLocationInfo(FVectorVMContext& Context, EVectorVMOperandLocation InLocation, int32 InPrimaryIndex, int32 InSecondaryIndex = INDEX_NONE, int32 InTertiaryIndex = INDEX_NONE);
+		FDataLocationInfo()
+			: Location(EVectorVMOperandLocation::Undefined)
+			, PrimaryLocationIndex(INDEX_NONE)
+			, SecondaryLocationIndex(INDEX_NONE)
+			, TertiaryLocationIndex(INDEX_NONE)
+			, MemoryAddress(NULL)
+		{}
+	};
+
+	struct VECTORVM_API FOpDebugInfo
+	{
+		EVectorVMOp Op;
+		EVMType OpType;
+		int32 NumArgs;
+
+		FDebugValue PreOpValues[NUM_VM_OP_DEBUG_VALUES];
+		FDebugValue PostOpValues[NUM_VM_OP_DEBUG_VALUES];
+		FDataLocationInfo LocationInfo[NUM_VM_OP_DEBUG_VALUES];
+	};
+
+	class VECTORVM_API FVectorVMDebugger
+	{
+	private:
+		FVectorVMDebuggerImpl* Impl;
+
+	public:
+		FVectorVMDebugger();
+		~FVectorVMDebugger();
+		void AddInstanceToGather(int32 Instance);
+		const TArray<FOpDebugInfo>* GetDebugInfoForInstance(int32 Instance);
+		FVectorVMDebuggerImpl* GetImpl();
+		void InitForScriptRun(int32 StartInstance, const TArray<int32> InstancesToDebug);
+	};
+
+	VECTORVM_API void AttachDebugger(FVectorVMDebugger* Debugger);
+	VECTORVM_API void DetachDebugger(FVectorVMDebugger* Debugger);
+#endif
 
 } // namespace VectorVM
 
 
+/**
+Encapsulates a view of a number of shared data buffers with a shared counter and size.
+This is temporary until we can use integers in the constants for the size and counter.
+*/
+class FVectorVMSharedDataView
+{
+private:
+	/** Array of data buffers.*/
+	TArray<FVector4*> Buffers;
+	/** Total size of buffer in vectors. */
+	int32 Size;
+	/** Counter used to acquire and release slots for reading and writing. */
+	int32 Counter;//Need to ensure this is atomically updated if threading withing a single VM script run. Though I don't think we should do that.
+
+	// Dummy read and write data for when the VM requests an invalid buffer pointer.
+	// This can happen because we have no branching currently. Writing to OOB indices must still continue but just not do any damage.
+	FVector4 DummyRead;
+	FVector4 DummyWrite;
+
+public:
+	FVectorVMSharedDataView(int32 InSize, int32 InCounter)
+		: Size(InSize)
+		, Counter(InCounter)
+		, DummyRead(FVector4(0.0f, 0.0f, 0.0f, 0.0f))
+	{
+	}
+
+	FORCEINLINE bool ValidIndex(int32 Index){ return Index >= 0 && Index < Size; }
+
+	FORCEINLINE FVector4* GetReadBuffer(int32 BufferIndex, int32 DataIndex) { return ValidIndex(DataIndex) && Buffers[BufferIndex] ? Buffers[BufferIndex] + DataIndex : &DummyRead; }
+	FORCEINLINE FVector4* GetWriteBuffer(int32 BufferIndex, int32 DataIndex) { return ValidIndex(DataIndex) && Buffers[BufferIndex] ? Buffers[BufferIndex] + DataIndex : &DummyWrite; }
+
+	FORCEINLINE void AddBuffer(FVector4* Buffer)
+	{
+		Buffers.Add(Buffer);
+	}
+
+	FORCEINLINE int32 GetCounter(){ return Counter; }
+	FORCEINLINE int32 GetSize(){ return Size; }
+
+	FORCEINLINE int32 AcquireIndexWrap(){ int32 Index = Counter; IncrementWrap(); return Index; }
+	FORCEINLINE int32 AcquireIndex(){ int32 Index = Counter; Increment(); return Index; }
+
+	FORCEINLINE int32 ConsumeIndexWrap(){ int32 Index = Counter; DecrementWrap(); return Index; }
+	FORCEINLINE int32 ConsumeIndex(){ int32 Index = Counter; Decrement(); return Index; }
+
+	FORCEINLINE void Increment(){ Counter = FMath::Min(Counter + 1, Size); }
+	FORCEINLINE void Decrement(){ Counter = FMath::Max(Counter - 1, 0); }
+
+	FORCEINLINE int32 Modulo(int32 X, int32 Max)
+	{
+		return Max == 0
+			? 0
+			: (X < 0 ? (X % Max) + Max : (X % Max));
+	}
+	FORCEINLINE void IncrementWrap(){ Counter = Size == 0 ? INDEX_NONE : Modulo(Counter + 1, Size - 1); }
+	FORCEINLINE void DecrementWrap(){ Counter = Size == 0 ? INDEX_NONE : Modulo(Counter - 1, Size - 1); }
+};
+
