Hopeful, blind compile fix for xbone.

diff --git a/Engine/Source/Runtime/VectorVM/Private/VectorVM.cpp b/Engine/Source/Runtime/VectorVM/Private/VectorVM.cpp
--- a/Engine/Source/Runtime/VectorVM/Private/VectorVM.cpp
+++ b/Engine/Source/Runtime/VectorVM/Private/VectorVM.cpp
@@ -85,8 +85,6 @@ UNiagaraSparseVolumeDataObject::UNiagaraSparseVolumeDataObject(const FObjectInit
 
 //////////////////////////////////////////////////////////////////////////
 
-FDummyHandler DummyHandler;
-
 #if WITH_EDITOR
 TArray<FString> OpNames;
 TArray<FString> OperandLocationNames;
@@ -165,7 +163,7 @@ struct FVectorVMContext
 		}
 	}
 
-	template<typename DstHandler, typename Arg0Handler, typename Arg1Handler = FDummyHandler, typename Arg2Handler = FDummyHandler, typename Arg3Handler = FDummyHandler>
+	template<typename DstHandler, typename Arg0Handler, typename Arg1Handle, typename Arg2Handler, typename Arg3Handle>
 	FORCEINLINE void PreOp(DstHandler& Dst, Arg0Handler& Arg0, Arg1Handler& Arg1 = DummyHandler, Arg2Handler& Arg2 = DummyHandler, Arg3Handler& Arg3 = DummyHandler)
 	{
 		if (AttachedDebugger)
@@ -174,7 +172,7 @@ struct FVectorVMContext
 		}
 	}
 
-	template<typename DstHandler, typename Arg0Handler, typename Arg1Handler = FDummyHandler, typename Arg2Handler = FDummyHandler, typename Arg3Handler = FDummyHandler>
+	template<typename DstHandler, typename Arg0Handler, typename Arg1Handler, typename Arg2Handler, typename Arg3Handler>
 	FORCEINLINE void PostOp(DstHandler& Dst, Arg0Handler& Arg0, Arg1Handler& Arg1 = DummyHandler, Arg2Handler& Arg2 = DummyHandler, Arg3Handler& Arg3 = DummyHandler)
 	{
 		if (AttachedDebugger)
@@ -185,11 +183,28 @@ struct FVectorVMContext
 #else
 	FORCEINLINE void BeginOp(VectorVM::EVMType InType, int32 InNumArgs, int32 InNumInstancesPerOp) { }
 
-	template<typename DstHandler, typename Arg0Handler, typename Arg1Handler = FDummyHandler, typename Arg2Handler = FDummyHandler, typename Arg3Handler = FDummyHandler>
-	FORCEINLINE void PreOp(DstHandler& Dst, Arg0Handler& Arg0, Arg1Handler& Arg1 = DummyHandler, Arg2Handler& Arg2 = DummyHandler, Arg3Handler& Arg3 = DummyHandler){	}
+	template<typename DstHandler, typename Arg0Handler, typename Arg1Handler, typename Arg2Handler, typename Arg3Handler>
+	FORCEINLINE void PreOp(DstHandler& Dst, Arg0Handler& Arg0, Arg1Handler& Arg1, Arg2Handler& Arg2, Arg3Handler& Arg3){	}
+	template<typename DstHandler, typename Arg0Handler, typename Arg1Handler, typename Arg2Handler, typename Arg3Handler>
+	FORCEINLINE void PostOp(DstHandler& Dst, Arg0Handler& Arg0, Arg1Handler& Arg1, Arg2Handler& Arg2, Arg3Handler& Arg3){	}
+
+
+	template<typename DstHandler, typename Arg0Handler, typename Arg1Handler, typename Arg2Handler>
+	FORCEINLINE void PreOp(DstHandler& Dst, Arg0Handler& Arg0, Arg1Handler& Arg1, Arg2Handler& Arg2){	}
+	template<typename DstHandler, typename Arg0Handler, typename Arg1Handler, typename Arg2Handler>
+	FORCEINLINE void PostOp(DstHandler& Dst, Arg0Handler& Arg0, Arg1Handler& Arg1, Arg2Handler& Arg2){	}
+
+
+	template<typename DstHandler, typename Arg0Handler, typename Arg1Handler>
+	FORCEINLINE void PreOp(DstHandler& Dst, Arg0Handler& Arg0, Arg1Handler& Arg1){	}
+	template<typename DstHandler, typename Arg0Handler, typename Arg1Handler>
+	FORCEINLINE void PostOp(DstHandler& Dst, Arg0Handler& Arg0, Arg1Handler& Arg1){	}
+
 
-	template<typename DstHandler, typename Arg0Handler, typename Arg1Handler = FDummyHandler, typename Arg2Handler = FDummyHandler, typename Arg3Handler = FDummyHandler>
-	FORCEINLINE void PostOp(DstHandler& Dst, Arg0Handler& Arg0, Arg1Handler& Arg1 = DummyHandler, Arg2Handler& Arg2 = DummyHandler, Arg3Handler& Arg3 = DummyHandler){	}
+	template<typename DstHandler, typename Arg0Handler>
+	FORCEINLINE void PreOp(DstHandler& Dst, Arg0Handler& Arg0){	}
+	template<typename DstHandler, typename Arg0Handler>
+	FORCEINLINE void PostOp(DstHandler& Dst, Arg0Handler& Arg0){	}
 #endif
 };
 
@@ -387,9 +402,9 @@ struct TUnaryKernel
 		Context.BeginOp(Kernel::Type, 1, Kernel::NumInstancesPerOp);
 		for (int32 i = 0; i < Context.NumInstances; i += Kernel::NumInstancesPerOp)
 		{
-			Context.PreOp(Dst, Arg0);
+			Context.PreOp<DstHandler, Arg0Handler>(Dst, Arg0);
 			Kernel::DoKernel(Dst.Get(), Arg0.Get());
-			Context.PostOp(Dst, Arg0);
+			Context.PostOp<DstHandler, Arg0Handler>(Dst, Arg0);
 			Dst.Advance();	Arg0.Advance();
 		}
 	}
@@ -406,9 +421,9 @@ struct TBinaryKernel
 		Context.BeginOp(Kernel::Type, 2, Kernel::NumInstancesPerOp);
 		for (int32 i = 0; i < Context.NumInstances; i += Kernel::NumInstancesPerOp)
 		{
-			Context.PreOp(Dst, Arg0, Arg1);
+			Context.PreOp<DstHandler, Arg0Handler, Arg1Handler>(Dst, Arg0, Arg1);
 			Kernel::DoKernel(Dst.Get(), Arg0.Get(), Arg1.Get());
-			Context.PostOp(Dst, Arg0, Arg1);
+			Context.PostOp<DstHandler, Arg0Handler, Arg1Handler>(Dst, Arg0, Arg1);
 			Dst.Advance(); Arg0.Advance(); Arg1.Advance();
 		}
 	}
@@ -426,9 +441,9 @@ struct TTrinaryKernel
 		Context.BeginOp(Kernel::Type, 3, Kernel::NumInstancesPerOp);
 		for (int32 i = 0; i < Context.NumInstances; i += Kernel::NumInstancesPerOp)
 		{
-			Context.PreOp(Dst, Arg0, Arg1, Arg2);
+			Context.PreOp<DstHandler, Arg0Handler, Arg1Handler, Arg2Handler>(Dst, Arg0, Arg1, Arg2);
 			Kernel::DoKernel(Dst.Get(), Arg0.Get(), Arg1.Get(), Arg2.Get());
-			Context.PostOp(Dst, Arg0, Arg1, Arg2);
+			Context.PostOp<DstHandler, Arg0Handler, Arg1Handler, Arg2Handler>(Dst, Arg0, Arg1, Arg2);
 			Dst.Advance(); Arg0.Advance(); Arg1.Advance(); Arg2.Advance();
 		}
 	}
@@ -447,9 +462,9 @@ struct TQuaternaryKernel
 		Context.BeginOp(Kernel::Type, 4, Kernel::NumInstancesPerOp);
 		for (int32 i = 0; i < Context.NumInstances; i += Kernel::NumInstancesPerOp)
 		{
-			Context.PreOp(Dst, Arg0, Arg1, Arg2);
+			Context.PreOp<DstHandler, Arg0Handler, Arg1Handler, Arg2Handler, Arg3Handler>(Dst, Arg0, Arg1, Arg2, Arg3);
 			Kernel::DoKernel(Dst.Get(), Arg0.Get(), Arg1.Get(), Arg2.Get(), Arg3.Get());
-			Context.PostOp(Dst, Arg0, Arg1, Arg2, Arg3);
+			Context.PostOp<DstHandler, Arg0Handler, Arg1Handler, Arg2Handler, Arg3Handler>(Dst, Arg0, Arg1, Arg2, Arg3);
 			Dst.Advance(); Arg0.Advance(); Arg1.Advance(); Arg2.Advance(); Arg3.Advance();
 		}
 	}
@@ -1255,13 +1270,13 @@ struct FKernelSharedDataGetAppendIndexBase
 		Context.BeginOp(VectorVM::EVMType::Vector4, 2, 1);
 		for (int32 i = 0; i < NumInstances; ++i)
 		{
-			Context.PreOp(IndexDest, ValidSrc, IdxHandler);
+			Context.PreOp<FRegisterDestHandler<int32, 4>, FRegisterHandler<float, 4>, IndexHandler>(IndexDest, ValidSrc, IdxHandler);
 
 			int32 Index = ValidSrc.Get() > 0.0f ? IdxHandler.GetNextIndex() : INDEX_NONE;
 			int32* Dest = IndexDest.Get();
 			Dest[0] = Index; Dest[1] = Index; Dest[2] = Index; Dest[3] = Index;
 
-			Context.PostOp(IndexDest, ValidSrc, IdxHandler);
+			Context.PostOp<FRegisterDestHandler<int32, 4>, FRegisterHandler<float, 4>, IndexHandler>(IndexDest, ValidSrc, IdxHandler);
 			ValidSrc.Advance();
 			IndexDest.Advance();
 			IdxHandler.Advance();
@@ -1287,13 +1302,13 @@ struct FKernelSharedDataGetConsumeIndexBase
 		Context.BeginOp(VectorVM::EVMType::Vector4, 1, 1);
 		for (int32 i = 0; i < NumInstances; ++i)
 		{
-			Context.PreOp(IndexDest, IdxHandler);
+			Context.PreOp<FRegisterDestHandler<int32, 4>, IndexHandler>(IndexDest, IdxHandler);
 
 			int32 Index = IdxHandler.GetNextIndex();
 			//Better to just stay in int pipeline?
 			*IndexDest.Get() = Index;//Only need index in X;
 
-			Context.PostOp(IndexDest, IdxHandler);
+			Context.PostOp<FRegisterDestHandler<int32, 4>, IndexHandler>(IndexDest, IdxHandler);
 			IndexDest.Advance();
 			IdxHandler.Advance();
 		}
@@ -1315,7 +1330,7 @@ struct FKernelSharedDataIndexValid
 		Context.BeginOp(VectorVM::EVMType::Vector4, 1, 1);
 		for (int32 i = 0; i < NumInstances; ++i)
 		{
-			Context.PreOp(ValidDest, IndexSrc);
+			Context.PreOp<FRegisterDestHandler<float, 4>, FRegisterHandler<int32, 4>>(ValidDest, IndexSrc);
 
 			int32 Index = IndexSrc.Get();
 			float Valid = SharedData.ValidIndex(Index) ? 1.0f : 0.0f;
@@ -1323,7 +1338,7 @@ struct FKernelSharedDataIndexValid
 			float* Dst = ValidDest.Get();
 			Dst[0] = Valid; Dst[1] = Valid; Dst[2] = Valid; Dst[3] = Valid;
 
-			Context.PostOp(ValidDest, IndexSrc);
+			Context.PostOp<FRegisterDestHandler<float, 4>, FRegisterHandler<int32, 4>>(ValidDest, IndexSrc);
 			ValidDest.Advance();
 			IndexSrc.Advance();;
 		}
diff --git a/Engine/Source/Runtime/VectorVM/Public/VectorVM.h b/Engine/Source/Runtime/VectorVM/Public/VectorVM.h
--- a/Engine/Source/Runtime/VectorVM/Public/VectorVM.h
+++ b/Engine/Source/Runtime/VectorVM/Public/VectorVM.h
@@ -141,6 +141,8 @@ namespace VectorVM
 	//////////////////////////////////////////////////////////////////////////
 	//Debugging
 
+	//VM Deubgger is currently broken. I wrote it and had it working on win 64 and unknowingly used some C++11 features not supported by the XBone compiler.
+	//I've had to hack the ENABLE_VM_DEBUGGING (0) code to get this to build but I've not done the ENABLE_VM_DEBUGGING (1) path as of yet.
 #define ENABLE_VM_DEBUGGING (0)
 //#define ENABLE_VM_DEBUGGING (1)
 
