Move VectorVM tests into separate file

diff --git a/Engine/Source/Runtime/VectorVM/Private/Tests/VectorVMTests.cpp b/Engine/Source/Runtime/VectorVM/Private/Tests/VectorVMTests.cpp
new file mode 100644
--- /dev/null
+++ b/Engine/Source/Runtime/VectorVM/Private/Tests/VectorVMTests.cpp
@@ -0,0 +1,106 @@
+// Copyright 1998-2015 Epic Games, Inc. All Rights Reserved.
+
+#include "Private/VectorVMPrivate.h"
+#include "VectorVM.h"
+
+#define OP_REGISTER (0)
+#define OP1_CONST (1 << 1)
+#define OP2_CONST (1 << 2)
+
+#define SRCOP_RRRR (OP_REGISTER | OP_REGISTER | OP_REGISTER | OP_REGISTER)
+#define SRCOP_RRCR (OP_REGISTER | OP_REGISTER | OP1_CONST | OP_REGISTER)
+#define SRCOP_RCCR (OP_REGISTER | OP2_CONST | OP1_CONST | OP_REGISTER)
+
+IMPLEMENT_SIMPLE_AUTOMATION_TEST(FVectorVMTest, "System.Core.Math.Vector VM", EAutomationTestFlags::ApplicationContextMask | EAutomationTestFlags::SmokeFilter)
+
+/*------------------------------------------------------------------------------
+Automation test for the VM.
+------------------------------------------------------------------------------*/
+bool FVectorVMTest::RunTest(const FString& Parameters)
+{
+	VectorVM::Init();
+
+	uint8 TestCode[] =
+	{
+		(uint8)EVectorVMOp::mul, SRCOP_RRRR, 0x00, 0x0 + VectorVM::NumTempRegisters, 0x0 + VectorVM::NumTempRegisters,       // mul r0, r8, r8
+		(uint8)EVectorVMOp::mad, SRCOP_RRRR, 0x01, 0x01 + VectorVM::NumTempRegisters, 0x01 + VectorVM::NumTempRegisters, 0x00, // mad r1, r9, r9, r0
+		(uint8)EVectorVMOp::mad, SRCOP_RRRR, 0x00, 0x02 + VectorVM::NumTempRegisters, 0x02 + VectorVM::NumTempRegisters, 0x01, // mad r0, r10, r10, r1
+		(uint8)EVectorVMOp::add, SRCOP_RRCR, 0x01, 0x00, 0x01,       // addi r1, r0, c1
+		(uint8)EVectorVMOp::neg, SRCOP_RRRR, 0x00, 0x01,             // neg r0, r1
+		(uint8)EVectorVMOp::clamp, SRCOP_RCCR, VectorVM::FirstOutputRegister, 0x00, 0x02, 0x03, // clampii r40, r0, c2, c3
+		0x00 // terminator
+	};
+
+	VectorRegister TestRegisters[4][VectorVM::VectorsPerChunk];
+	VectorRegister* InputRegisters[3] ={ TestRegisters[0], TestRegisters[1], TestRegisters[2] };
+	VectorRegister* OutputRegisters[1] ={ TestRegisters[3] };
+
+	VectorRegister Inputs[3][VectorVM::VectorsPerChunk];
+	for (int32 i = 0; i < VectorVM::ChunkSize; i++)
+	{
+		reinterpret_cast<float*>(&Inputs[0])[i] = static_cast<float>(i);
+		reinterpret_cast<float*>(&Inputs[1])[i] = static_cast<float>(i);
+		reinterpret_cast<float*>(&Inputs[2])[i] = static_cast<float>(i);
+		reinterpret_cast<float*>(InputRegisters[0])[i] = static_cast<float>(i);
+		reinterpret_cast<float*>(InputRegisters[1])[i] = static_cast<float>(i);
+		reinterpret_cast<float*>(InputRegisters[2])[i] = static_cast<float>(i);
+	}
+
+	FVector4 ConstantTable[VectorVM::MaxConstants];
+	ConstantTable[0] = FVector4(0.0f, 0.0f, 0.0f, 0.0f);
+	ConstantTable[1] = FVector4(5.0f, 5.0f, 5.0f, 5.0f);
+	ConstantTable[2] = FVector4(-20.0f, -20.0f, -20.0f, -20.0f);
+	ConstantTable[3] = FVector4(20.0f, 20.0f, 20.0f, 20.0f);
+
+	FVectorVMSharedDataView* dummy = NULL;
+
+	VectorVM::Exec(
+		TestCode,
+		InputRegisters, 3,
+		OutputRegisters, 1,
+		ConstantTable,
+		nullptr,
+		dummy,
+		VectorVM::VectorsPerChunk
+		);
+
+	for (int32 i = 0; i < VectorVM::ChunkSize; i++)
+	{
+		float Ins[3];
+
+		// Verify that the input registers were not overwritten.
+		for (int32 InputIndex = 0; InputIndex < 3; ++InputIndex)
+		{
+			float In = Ins[InputIndex] = reinterpret_cast<float*>(&Inputs[InputIndex])[i];
+			float R = reinterpret_cast<float*>(InputRegisters[InputIndex])[i];
+			if (In != R)
+			{
+				AddError(FString::Printf(TEXT("Input register %d vector %d element %d overwritten. Has %f expected %f"),
+					InputIndex,
+					i / VectorVM::ElementsPerVector,
+					i % VectorVM::ElementsPerVector,
+					R,
+					In
+					));
+				return false;
+			}
+		}
+
+		// Verify that outputs match what we expect.
+		float Out = reinterpret_cast<float*>(OutputRegisters[0])[i];
+		float Expected = FMath::Clamp<float>(-(Ins[0] * Ins[0] + Ins[1] * Ins[1] + Ins[2] * Ins[2] + 5.0f), -20.0f, 20.0f);
+		if (Out != Expected)
+		{
+			AddError(FString::Printf(TEXT("Output register %d vector %d element %d is wrong. Has %f expected %f"),
+				0,
+				i / VectorVM::ElementsPerVector,
+				i % VectorVM::ElementsPerVector,
+				Out,
+				Expected
+				));
+			return false;
+		}
+	}
+
+	return true;
+}
\ No newline at end of file
diff --git a/Engine/Source/Runtime/VectorVM/Private/VectorVM.cpp b/Engine/Source/Runtime/VectorVM/Private/VectorVM.cpp
--- a/Engine/Source/Runtime/VectorVM/Private/VectorVM.cpp
+++ b/Engine/Source/Runtime/VectorVM/Private/VectorVM.cpp
@@ -1606,101 +1606,6 @@ FString VectorVM::GetOperandLocationName(EVectorVMOperandLocation Location)
 }
 #endif
 
-/*------------------------------------------------------------------------------
-	Automation test for the VM.
-------------------------------------------------------------------------------*/
-
-IMPLEMENT_SIMPLE_AUTOMATION_TEST(FVectorVMTest, "System.Core.Math.Vector VM", EAutomationTestFlags::ApplicationContextMask | EAutomationTestFlags::SmokeFilter)
-
-bool FVectorVMTest::RunTest(const FString& Parameters)
-{
-	VectorVM::Init();
-
-	uint8 TestCode[] =
-	{
-		(uint8)EVectorVMOp::mul, SRCOP_RRRR, 0x00, 0x0 + VectorVM::NumTempRegisters, 0x0 + VectorVM::NumTempRegisters,       // mul r0, r8, r8
-		(uint8)EVectorVMOp::mad, SRCOP_RRRR, 0x01, 0x01 + VectorVM::NumTempRegisters, 0x01 + VectorVM::NumTempRegisters, 0x00, // mad r1, r9, r9, r0
-		(uint8)EVectorVMOp::mad, SRCOP_RRRR, 0x00, 0x02 + VectorVM::NumTempRegisters, 0x02 + VectorVM::NumTempRegisters, 0x01, // mad r0, r10, r10, r1
-		(uint8)EVectorVMOp::add, SRCOP_RRCR, 0x01, 0x00, 0x01,       // addi r1, r0, c1
-		(uint8)EVectorVMOp::neg, SRCOP_RRRR, 0x00, 0x01,             // neg r0, r1
-		(uint8)EVectorVMOp::clamp, SRCOP_RCCR, VectorVM::FirstOutputRegister, 0x00, 0x02, 0x03, // clampii r40, r0, c2, c3
-		0x00 // terminator
-	};
-
-	VectorRegister TestRegisters[4][VectorVM::VectorsPerChunk];
-	VectorRegister* InputRegisters[3] = { TestRegisters[0], TestRegisters[1], TestRegisters[2] };
-	VectorRegister* OutputRegisters[1] = { TestRegisters[3] };
-
-	VectorRegister Inputs[3][VectorVM::VectorsPerChunk];
-	for (int32 i = 0; i < VectorVM::ChunkSize; i++)
-	{
-		reinterpret_cast<float*>(&Inputs[0])[i] = static_cast<float>(i);
-		reinterpret_cast<float*>(&Inputs[1])[i] = static_cast<float>(i);
-		reinterpret_cast<float*>(&Inputs[2])[i] = static_cast<float>(i);
-		reinterpret_cast<float*>(InputRegisters[0])[i] = static_cast<float>(i);
-		reinterpret_cast<float*>(InputRegisters[1])[i] = static_cast<float>(i);
-		reinterpret_cast<float*>(InputRegisters[2])[i] = static_cast<float>(i);
-	}
-
-	FVector4 ConstantTable[VectorVM::MaxConstants];
-	ConstantTable[0] = FVector4(0.0f, 0.0f, 0.0f, 0.0f);
-	ConstantTable[1] = FVector4(5.0f, 5.0f, 5.0f, 5.0f);
-	ConstantTable[2] = FVector4(-20.0f, -20.0f, -20.0f, -20.0f);
-	ConstantTable[3] = FVector4(20.0f, 20.0f, 20.0f, 20.0f);
-
-	FVectorVMSharedDataView* dummy = NULL;
-
-	VectorVM::Exec(
-		TestCode,
-		InputRegisters, 3,
-		OutputRegisters, 1,
-		ConstantTable,
-		nullptr,
-		dummy,
-		VectorVM::VectorsPerChunk
-		);
-
-	for (int32 i = 0; i < VectorVM::ChunkSize; i++)
-	{
-		float Ins[3];
-
-		// Verify that the input registers were not overwritten.
-		for (int32 InputIndex = 0; InputIndex < 3; ++InputIndex)
-		{
-			float In = Ins[InputIndex] = reinterpret_cast<float*>(&Inputs[InputIndex])[i];
-			float R = reinterpret_cast<float*>(InputRegisters[InputIndex])[i];
-			if (In != R)
-			{
-				UE_LOG(LogVectorVM,Error,TEXT("Input register %d vector %d element %d overwritten. Has %f expected %f"),
-					InputIndex,
-					i / VectorVM::ElementsPerVector,
-					i % VectorVM::ElementsPerVector,
-					R,
-					In
-					);
-				return false;
-			}
-		}
-
-		// Verify that outputs match what we expect.
-		float Out = reinterpret_cast<float*>(OutputRegisters[0])[i];
-		float Expected = FMath::Clamp<float>(-(Ins[0] * Ins[0] + Ins[1] * Ins[1] + Ins[2] * Ins[2] + 5.0f), -20.0f, 20.0f);
-		if (Out != Expected)
-		{
-			UE_LOG(LogVectorVM,Error,TEXT("Output register %d vector %d element %d is wrong. Has %f expected %f"),
-				0,
-				i / VectorVM::ElementsPerVector,
-				i % VectorVM::ElementsPerVector,
-				Out,
-				Expected
-				);
-			return false;
-		}
-	}
-
-	return true;
-}
-
 //////////////////////////////////////////////////////////////////////////
 #if ENABLE_VM_DEBUGGING
 
