Copying //UE4/Dev-Rendering to //UE4/Dev-Main (Source: //UE4/Dev-Rendering @ 3045398)

diff --git a/Engine/Source/Runtime/Renderer/Private/ShadowRendering.h b/Engine/Source/Runtime/Renderer/Private/ShadowRendering.h
--- a/Engine/Source/Runtime/Renderer/Private/ShadowRendering.h
+++ b/Engine/Source/Runtime/Renderer/Private/ShadowRendering.h
@@ -10,6 +10,120 @@
 #include "ShaderParameterUtils.h"
 #include "SceneCore.h"
 
+
+/** Uniform buffer for rendering deferred lights. */
+BEGIN_UNIFORM_BUFFER_STRUCT(FDeferredLightUniformStruct,)
+	DECLARE_UNIFORM_BUFFER_STRUCT_MEMBER(FVector,LightPosition)
+	DECLARE_UNIFORM_BUFFER_STRUCT_MEMBER(float,LightInvRadius)
+	DECLARE_UNIFORM_BUFFER_STRUCT_MEMBER(FVector,LightColor)
+	DECLARE_UNIFORM_BUFFER_STRUCT_MEMBER(float,LightFalloffExponent)
+	DECLARE_UNIFORM_BUFFER_STRUCT_MEMBER(FVector,NormalizedLightDirection)
+	DECLARE_UNIFORM_BUFFER_STRUCT_MEMBER(FVector2D,SpotAngles)
+	DECLARE_UNIFORM_BUFFER_STRUCT_MEMBER(float,SourceRadius)
+	DECLARE_UNIFORM_BUFFER_STRUCT_MEMBER(float,SourceLength)
+	DECLARE_UNIFORM_BUFFER_STRUCT_MEMBER(float,MinRoughness)
+	DECLARE_UNIFORM_BUFFER_STRUCT_MEMBER(float,ContactShadowLength)
+	DECLARE_UNIFORM_BUFFER_STRUCT_MEMBER(FVector2D,DistanceFadeMAD)
+	DECLARE_UNIFORM_BUFFER_STRUCT_MEMBER(FVector4,ShadowMapChannelMask)
+	DECLARE_UNIFORM_BUFFER_STRUCT_MEMBER(uint32,ShadowedBits)
+	DECLARE_UNIFORM_BUFFER_STRUCT_MEMBER(uint32,LightingChannelMask)
+END_UNIFORM_BUFFER_STRUCT(FDeferredLightUniformStruct)
+
+extern float GMinScreenRadiusForLights;
+extern uint32 GetShadowQuality();
+
+template<typename ShaderRHIParamRef>
+void SetDeferredLightParameters(
+	FRHICommandList& RHICmdList, 
+	const ShaderRHIParamRef ShaderRHI, 
+	const TShaderUniformBufferParameter<FDeferredLightUniformStruct>& DeferredLightUniformBufferParameter, 
+	const FLightSceneInfo* LightSceneInfo,
+	const FSceneView& View)
+{
+	FVector4 LightPositionAndInvRadius;
+	FVector4 LightColorAndFalloffExponent;
+		
+	FDeferredLightUniformStruct DeferredLightUniformsValue;
+
+	// Get the light parameters
+	LightSceneInfo->Proxy->GetParameters(
+		LightPositionAndInvRadius,
+		LightColorAndFalloffExponent,
+		DeferredLightUniformsValue.NormalizedLightDirection,
+		DeferredLightUniformsValue.SpotAngles,
+		DeferredLightUniformsValue.SourceRadius,
+		DeferredLightUniformsValue.SourceLength,
+		DeferredLightUniformsValue.MinRoughness);
+	
+	DeferredLightUniformsValue.LightPosition = LightPositionAndInvRadius;
+	DeferredLightUniformsValue.LightInvRadius = LightPositionAndInvRadius.W;
+	DeferredLightUniformsValue.LightColor = LightColorAndFalloffExponent;
+	DeferredLightUniformsValue.LightFalloffExponent = LightColorAndFalloffExponent.W;
+
+	const FVector2D FadeParams = LightSceneInfo->Proxy->GetDirectionalLightDistanceFadeParameters(View.GetFeatureLevel(), LightSceneInfo->IsPrecomputedLightingValid());
+
+	// use MAD for efficiency in the shader
+	DeferredLightUniformsValue.DistanceFadeMAD = FVector2D(FadeParams.Y, -FadeParams.X * FadeParams.Y);
+
+	int32 ShadowMapChannel = LightSceneInfo->Proxy->GetShadowMapChannel();
+
+	static const auto AllowStaticLightingVar = IConsoleManager::Get().FindTConsoleVariableDataInt(TEXT("r.AllowStaticLighting"));
+	const bool bAllowStaticLighting = (!AllowStaticLightingVar || AllowStaticLightingVar->GetValueOnRenderThread() != 0);
+
+	if (!bAllowStaticLighting)
+	{
+		ShadowMapChannel = INDEX_NONE;
+	}
+
+	DeferredLightUniformsValue.ShadowMapChannelMask = FVector4(
+		ShadowMapChannel == 0 ? 1 : 0,
+		ShadowMapChannel == 1 ? 1 : 0,
+		ShadowMapChannel == 2 ? 1 : 0,
+		ShadowMapChannel == 3 ? 1 : 0);
+
+	const bool bDynamicShadows = View.Family->EngineShowFlags.DynamicShadows && GetShadowQuality() > 0;
+	const bool bHasLightFunction = LightSceneInfo->Proxy->GetLightFunctionMaterial() != NULL;
+	DeferredLightUniformsValue.ShadowedBits  = LightSceneInfo->Proxy->CastsStaticShadow() || bHasLightFunction ? 1 : 0;
+	DeferredLightUniformsValue.ShadowedBits |= LightSceneInfo->Proxy->CastsDynamicShadow() && View.Family->EngineShowFlags.DynamicShadows ? 3 : 0;
+
+	static auto* ContactShadowsCVar = IConsoleManager::Get().FindTConsoleVariableDataInt(TEXT("r.ContactShadows"));
+	DeferredLightUniformsValue.ContactShadowLength = 0;
+
+	if (ContactShadowsCVar && ContactShadowsCVar->GetValueOnRenderThread() != 0)
+	{
+		DeferredLightUniformsValue.ContactShadowLength = LightSceneInfo->Proxy->GetContactShadowLength();
+	}
+
+	if( LightSceneInfo->Proxy->IsInverseSquared() )
+	{
+		// Correction for lumen units
+		DeferredLightUniformsValue.LightColor *= 16.0f;
+	}
+
+	// When rendering reflection captures, the direct lighting of the light is actually the indirect specular from the main view
+	if (View.bIsReflectionCapture)
+	{
+		DeferredLightUniformsValue.LightColor *= LightSceneInfo->Proxy->GetIndirectLightingScale();
+	}
+
+	const ELightComponentType LightType = (ELightComponentType)LightSceneInfo->Proxy->GetLightType();
+
+	if (LightType == LightType_Point || LightType == LightType_Spot)
+	{
+		// Distance fade
+		FSphere Bounds = LightSceneInfo->Proxy->GetBoundingSphere();
+
+		const float DistanceSquared = ( Bounds.Center - View.ViewMatrices.ViewOrigin ).SizeSquared();
+		float Fade = FMath::Square( FMath::Min( 0.0002f, GMinScreenRadiusForLights / Bounds.W ) * View.LODDistanceFactor ) * DistanceSquared;
+		Fade = FMath::Clamp( 6.0f - 6.0f * Fade, 0.0f, 1.0f );
+		DeferredLightUniformsValue.LightColor *= Fade;
+	}
+
+	DeferredLightUniformsValue.LightingChannelMask = LightSceneInfo->Proxy->GetLightingChannelMask();
+
+	SetUniformBufferParameterImmediate(RHICmdList, ShaderRHI,DeferredLightUniformBufferParameter,DeferredLightUniformsValue);
+}
+
 // Forward declarations.
 class FProjectedShadowInfo;
 
@@ -720,16 +834,26 @@ public:
 	/** Renders shadow maps for translucent primitives. */
 	void RenderTranslucencyDepths(FRHICommandList& RHICmdList, class FSceneRenderer* SceneRenderer);
 
+	static void SetBlendStateForProjection(
+		FRHICommandListImmediate& RHICmdList,
+		int32 ShadowMapChannel,
+		bool bIsWholeSceneDirectionalShadow,
+		bool bUseFadePlane,
+		bool bProjectingForForwardShading,
+		bool bMobileModulatedProjections);
+
+	void SetBlendStateForProjection(FRHICommandListImmediate& RHICmdList, bool bProjectingForForwardShading, bool bMobileModulatedProjections) const;
+
 	/**
 	 * Projects the shadow onto the scene for a particular view.
 	 */
-	void RenderProjection(FRHICommandListImmediate& RHICmdList, int32 ViewIndex, const class FViewInfo* View, bool bMobile) const;
+	void RenderProjection(FRHICommandListImmediate& RHICmdList, int32 ViewIndex, const class FViewInfo* View, bool bProjectingForForwardShading, bool bMobile) const;
 
 	/** Renders ray traced distance field shadows. */
-	void RenderRayTracedDistanceFieldProjection(FRHICommandListImmediate& RHICmdList, const class FViewInfo& View) const;
+	void RenderRayTracedDistanceFieldProjection(FRHICommandListImmediate& RHICmdList, const class FViewInfo& View, bool bProjectingForForwardShading) const;
 
 	/** Render one pass point light shadow projections. */
-	void RenderOnePassPointLightProjection(FRHICommandListImmediate& RHICmdList, int32 ViewIndex, const FViewInfo& View) const;
+	void RenderOnePassPointLightProjection(FRHICommandListImmediate& RHICmdList, int32 ViewIndex, const FViewInfo& View, bool bProjectingForForwardShading) const;
 
 	/**
 	 * Renders the projected shadow's frustum wireframe with the given FPrimitiveDrawInterface.
@@ -809,7 +933,7 @@ public:
 	void SortSubjectMeshElements();
 
 	// 0 if Setup...() wasn't called yet
-	const FLightSceneInfo & GetLightSceneInfo() const { return *LightSceneInfo; }
+	const FLightSceneInfo& GetLightSceneInfo() const { return *LightSceneInfo; }
 	const FLightSceneInfoCompact& GetLightSceneInfoCompact() const { return LightSceneInfoCompact; }
 	/**
 	 * Parent primitive of the shadow group that created this shadow, if not a bWholeSceneShadow.
@@ -1329,6 +1453,13 @@ public:
 
 		SetShaderValue(RHICmdList, ShaderRHI, ShadowFadeFraction, ShadowInfo->FadeAlphas[ViewIndex] );
 		SetShaderValue(RHICmdList, ShaderRHI, ShadowSharpen, ShadowInfo->GetLightSceneInfo().Proxy->GetShadowSharpen() * 7.0f + 1.0f );
+
+		auto DeferredLightParameter = GetUniformBufferParameter<FDeferredLightUniformStruct>();
+
+		if (DeferredLightParameter.IsBound())
+		{
+			SetDeferredLightParameters(RHICmdList, ShaderRHI, DeferredLightParameter, &ShadowInfo->GetLightSceneInfo(), View);
+		}
 	}
 
 	/**
@@ -1654,6 +1785,13 @@ public:
 		SetShaderValue(RHICmdList, ShaderRHI, PointLightDepthBiasParameters, FVector2D(ShadowInfo->GetShaderDepthBias(), 0.0f));
 
 		SetSamplerParameter(RHICmdList, ShaderRHI, ShadowDepthTextureSampler, TStaticSamplerState<SF_Point, AM_Clamp, AM_Clamp, AM_Clamp>::GetRHI());
+
+		auto DeferredLightParameter = GetUniformBufferParameter<FDeferredLightUniformStruct>();
+
+		if (DeferredLightParameter.IsBound())
+		{
+			SetDeferredLightParameters(RHICmdList, ShaderRHI, DeferredLightParameter, &ShadowInfo->GetLightSceneInfo(), View);
+		}
 	}
 
 	virtual bool Serialize(FArchive& Ar) override
diff --git a/Engine/Source/Runtime/Renderer/Private/ShadowSetup.cpp b/Engine/Source/Runtime/Renderer/Private/ShadowSetup.cpp
--- a/Engine/Source/Runtime/Renderer/Private/ShadowSetup.cpp
+++ b/Engine/Source/Runtime/Renderer/Private/ShadowSetup.cpp
@@ -34,6 +34,14 @@ FAutoConsoleVariableRef CVarCacheWholeSceneShadows(
 	ECVF_Scalability | ECVF_RenderThreadSafe
 	);
 
+int32 GWholeSceneShadowCacheMb = 150;
+FAutoConsoleVariableRef CVarWholeSceneShadowCacheMb(
+	TEXT("r.Shadow.WholeSceneShadowCacheMb"),
+	GWholeSceneShadowCacheMb,
+	TEXT("Amount of memory that can be spent caching whole scene shadows.  ShadowMap allocations in a single frame can cause this to be exceeded."),
+	ECVF_Scalability | ECVF_RenderThreadSafe
+	);
+
 int32 GCachedShadowsCastFromMovablePrimitives = 1;
 FAutoConsoleVariableRef CVarCachedWholeSceneShadowsCastFromMovablePrimitives(
 	TEXT("r.Shadow.CachedShadowsCastFromMovablePrimitives"),
@@ -1757,10 +1765,21 @@ void ComputeWholeSceneShadowCacheModes(
 				}
 				else
 				{
-					OutNumShadowMaps = 2;
-					// Note: ShadowMap with static primitives rendered first so movable shadowmap can composite
-					OutCacheModes[0] = SDCM_StaticPrimitivesOnly;
-					OutCacheModes[1] = SDCM_MovablePrimitivesOnly;
+					int64 CachedShadowMapsSize = Scene->GetCachedWholeSceneShadowMapsSize();
+
+					if (CachedShadowMapsSize < GWholeSceneShadowCacheMb * 1024 * 1024)
+					{
+						OutNumShadowMaps = 2;
+						// Note: ShadowMap with static primitives rendered first so movable shadowmap can composite
+						OutCacheModes[0] = SDCM_StaticPrimitivesOnly;
+						OutCacheModes[1] = SDCM_MovablePrimitivesOnly;
+					}
+					else
+					{
+						OutNumShadowMaps = 1;
+						OutCacheModes[0] = SDCM_Uncached;
+						CachedShadowMapData->ShadowMap.DepthTarget = NULL;
+					}
 				}
 			}
 			else
@@ -1775,12 +1794,22 @@ void ComputeWholeSceneShadowCacheModes(
 		}
 		else
 		{
-			OutNumShadowMaps = 2;
-			// Note: ShadowMap with static primitives rendered first so movable shadowmap can composite
-			OutCacheModes[0] = SDCM_StaticPrimitivesOnly;
-			OutCacheModes[1] = SDCM_MovablePrimitivesOnly;
+			int64 CachedShadowMapsSize = Scene->GetCachedWholeSceneShadowMapsSize();
+
+			if (CachedShadowMapsSize < GWholeSceneShadowCacheMb * 1024 * 1024)
+			{
+				OutNumShadowMaps = 2;
+				// Note: ShadowMap with static primitives rendered first so movable shadowmap can composite
+				OutCacheModes[0] = SDCM_StaticPrimitivesOnly;
+				OutCacheModes[1] = SDCM_MovablePrimitivesOnly;
 
-			Scene->CachedShadowMaps.Add(LightSceneInfo->Id, FCachedShadowMapData(ProjectedShadowInitializer, RealTime));
+				Scene->CachedShadowMaps.Add(LightSceneInfo->Id, FCachedShadowMapData(ProjectedShadowInitializer, RealTime));
+			}
+			else
+			{
+				OutNumShadowMaps = 1;
+				OutCacheModes[0] = SDCM_Uncached;
+			}
 		}
 	}
 	else
@@ -2640,8 +2669,15 @@ void FSceneRenderer::AllocateShadowDepthTargets(FRHICommandListImmediate& RHICmd
 
 			if (FeatureLevel < ERHIFeatureLevel::SM4 
 				// Mobile renderer only supports opaque per-object shadows or CSM
-				&& (!ProjectedShadowInfo->bPerObjectOpaqueShadow || !(ProjectedShadowInfo->bDirectionalLight && ProjectedShadowInfo->bWholeSceneShadow)))
+				&& (!ProjectedShadowInfo->bPerObjectOpaqueShadow && !(ProjectedShadowInfo->bDirectionalLight && ProjectedShadowInfo->bWholeSceneShadow)))
+			{
+				bShadowIsVisible = false;
+			}
+
+			if (IsForwardShadingEnabled(FeatureLevel) 
+				&& (!ProjectedShadowInfo->GetLightSceneInfo().Proxy->HasStaticShadowing() || ProjectedShadowInfo->GetLightSceneInfo().Proxy->GetPreviewShadowMapChannel() == -1))
 			{
+				// With forward shading, dynamic shadows are projected into channels of the light attenuation texture based on their assigned ShadowMapChannel
 				bShadowIsVisible = false;
 			}
 
@@ -2763,30 +2799,14 @@ void FSceneRenderer::AllocateShadowDepthTargets(FRHICommandListImmediate& RHICmd
 	{
 		FCachedShadowMapData& ShadowMapData = CachedShadowMapIt.Value();
 
-		if (ShadowMapData.ShadowMap.IsValid() && ViewFamily.CurrentRealTime - ShadowMapData.LastUsedTime > 5.0f)
+		if (ShadowMapData.ShadowMap.IsValid() && ViewFamily.CurrentRealTime - ShadowMapData.LastUsedTime > 2.0f)
 		{
 			ShadowMapData.ShadowMap.Release();
 		}
 	}
 
-#if STATS
-	int64 CachedShadowmapMemory = 0;
-
-	for (TMap<int32, FCachedShadowMapData>::TConstIterator CachedShadowMapIt(Scene->CachedShadowMaps); CachedShadowMapIt; ++CachedShadowMapIt)
-	{
-		const FCachedShadowMapData& ShadowMapData = CachedShadowMapIt.Value();
-
-		if (ShadowMapData.ShadowMap.IsValid())
-		{
-			CachedShadowmapMemory += ShadowMapData.ShadowMap.ComputeMemorySize();
-		}
-	}
-
-	SET_MEMORY_STAT(STAT_CachedShadowmapMemory, CachedShadowmapMemory);
-
-	int64 ReferencedShadowmapMemory = SortedShadowsForShadowDepthPass.ComputeMemorySize();
-	SET_MEMORY_STAT(STAT_ShadowmapAtlasMemory, ReferencedShadowmapMemory);
-#endif
+	SET_MEMORY_STAT(STAT_CachedShadowmapMemory, Scene->GetCachedWholeSceneShadowMapsSize());
+	SET_MEMORY_STAT(STAT_ShadowmapAtlasMemory, SortedShadowsForShadowDepthPass.ComputeMemorySize());
 }
 
 void FSceneRenderer::AllocatePerObjectShadowDepthTargets(FRHICommandListImmediate& RHICmdList, TArray<FProjectedShadowInfo*, SceneRenderingAllocator>& Shadows)
diff --git a/Engine/Source/Runtime/Renderer/Private/TiledDeferredLightRendering.cpp b/Engine/Source/Runtime/Renderer/Private/TiledDeferredLightRendering.cpp
--- a/Engine/Source/Runtime/Renderer/Private/TiledDeferredLightRendering.cpp
+++ b/Engine/Source/Runtime/Renderer/Private/TiledDeferredLightRendering.cpp
@@ -194,7 +194,7 @@ public:
 					LightData2.LightDirectionAndSpotlightMaskAndMinRoughness[LightIndex] = FVector4(NormalizedLightDirection, W);
 				}
 
-				LightData2.SpotAnglesAndSourceRadiusAndSimpleLighting[LightIndex] = FVector4(SpotAngles.X, SpotAngles.Y, SourceRadius, 0);
+				LightData2.SpotAnglesAndSourceRadiusAndSimpleLighting[LightIndex] = FVector4(SpotAngles.X, SpotAngles.Y, SourceRadius, SourceLength);
 
 				int32 ShadowMapChannel = LightSceneInfo->Proxy->GetShadowMapChannel();
 
diff --git a/Engine/Source/Runtime/Renderer/Private/TranslucentLighting.cpp b/Engine/Source/Runtime/Renderer/Private/TranslucentLighting.cpp
--- a/Engine/Source/Runtime/Renderer/Private/TranslucentLighting.cpp
+++ b/Engine/Source/Runtime/Renderer/Private/TranslucentLighting.cpp
@@ -1109,6 +1109,132 @@ void FDeferredShadingSceneRenderer::ClearTranslucentVolumeLighting(FRHICommandLi
 	}
 }
 
+class FClearTranslucentLightingVolumeCS : public FGlobalShader
+{
+	DECLARE_SHADER_TYPE(FClearTranslucentLightingVolumeCS, Global)
+public:
+
+	static const int32 CLEAR_BLOCK_SIZE = 4;
+
+	static bool ShouldCache(EShaderPlatform Platform)
+	{
+		return IsFeatureLevelSupported(Platform, ERHIFeatureLevel::SM5);
+	}
+
+	static void ModifyCompilationEnvironment(EShaderPlatform Platform, FShaderCompilerEnvironment& OutEnvironment)
+	{
+		FGlobalShader::ModifyCompilationEnvironment(Platform, OutEnvironment);
+		OutEnvironment.SetDefine(TEXT("CLEAR_COMPUTE_SHADER"), 1);
+		OutEnvironment.SetDefine(TEXT("CLEAR_BLOCK_SIZE"), CLEAR_BLOCK_SIZE);
+	}
+
+	FClearTranslucentLightingVolumeCS(const ShaderMetaType::CompiledShaderInitializerType& Initializer)
+		: FGlobalShader(Initializer)
+	{
+		Ambient0.Bind(Initializer.ParameterMap, TEXT("Ambient0"));
+		Directional0.Bind(Initializer.ParameterMap, TEXT("Directional0"));
+		Ambient1.Bind(Initializer.ParameterMap, TEXT("Ambient1"));
+		Directional1.Bind(Initializer.ParameterMap, TEXT("Directional1"));
+	}
+
+	FClearTranslucentLightingVolumeCS()
+	{
+	}
+
+	void SetParameters(
+		FRHIAsyncComputeCommandListImmediate& RHICmdList,
+		FUnorderedAccessViewRHIParamRef* VolumeUAVs,
+		int32 NumUAVs
+	)
+	{
+		check(NumUAVs == 4);
+		const FComputeShaderRHIParamRef ShaderRHI = GetComputeShader();
+		Ambient0.SetTexture(RHICmdList, ShaderRHI, NULL, VolumeUAVs[0]);
+		Directional0.SetTexture(RHICmdList, ShaderRHI, NULL, VolumeUAVs[1]);
+		Ambient1.SetTexture(RHICmdList, ShaderRHI, NULL, VolumeUAVs[2]);
+		Directional1.SetTexture(RHICmdList, ShaderRHI, NULL, VolumeUAVs[3]);
+	}
+
+	void UnsetParameters(FRHIAsyncComputeCommandListImmediate& RHICmdList)
+	{
+		const FComputeShaderRHIParamRef ShaderRHI = GetComputeShader();
+		Ambient0.UnsetUAV(RHICmdList, ShaderRHI);
+		Directional0.UnsetUAV(RHICmdList, ShaderRHI);
+		Ambient1.UnsetUAV(RHICmdList, ShaderRHI);
+		Directional1.UnsetUAV(RHICmdList, ShaderRHI);
+	}
+
+	virtual bool Serialize(FArchive& Ar) override
+	{
+		bool bShaderHasOutdatedParameters = FGlobalShader::Serialize(Ar);
+		Ar << Ambient0;
+		Ar << Directional0;
+		Ar << Ambient1;
+		Ar << Directional1;
+		return bShaderHasOutdatedParameters;
+	}
+
+private:
+	FRWShaderParameter Ambient0;
+	FRWShaderParameter Directional0;
+	FRWShaderParameter Ambient1;
+	FRWShaderParameter Directional1;
+};
+
+IMPLEMENT_SHADER_TYPE(, FClearTranslucentLightingVolumeCS, TEXT("TranslucentLightInjectionShaders"), TEXT("ClearTranslucentLightingVolumeCS"), SF_Compute)
+
+void FDeferredShadingSceneRenderer::ClearTranslucentVolumeLightingAsyncCompute(FRHICommandListImmediate& RHICmdList)
+{
+	FSceneRenderTargets& SceneContext = FSceneRenderTargets::Get(RHICmdList);
+
+	ensure(FeatureLevel >= ERHIFeatureLevel::SM5);
+	ensure(SceneContext.TranslucencyLightingVolumeAmbient[0]);
+	ensure(SceneContext.TranslucencyLightingVolumeDirectional[0]);
+	ensure(SceneContext.TranslucencyLightingVolumeAmbient[1]);
+	ensure(SceneContext.TranslucencyLightingVolumeDirectional[1]);
+
+	FUnorderedAccessViewRHIParamRef VolumeUAVs[4] = {
+	SceneContext.TranslucencyLightingVolumeAmbient[0]->GetRenderTargetItem().UAV,
+	SceneContext.TranslucencyLightingVolumeDirectional[0]->GetRenderTargetItem().UAV,
+	SceneContext.TranslucencyLightingVolumeAmbient[1]->GetRenderTargetItem().UAV,
+	SceneContext.TranslucencyLightingVolumeDirectional[1]->GetRenderTargetItem().UAV
+	};
+
+	FClearTranslucentLightingVolumeCS* ComputeShader = *TShaderMapRef<FClearTranslucentLightingVolumeCS>(GetGlobalShaderMap(FeatureLevel));
+	static const FName EndComputeFenceName(TEXT("TranslucencyLightingVolumeClearEndComputeFence"));
+	TranslucencyLightingVolumeClearEndFence = RHICmdList.CreateComputeFence(EndComputeFenceName);
+
+	static const FName BeginComputeFenceName(TEXT("TranslucencyLightingVolumeClearBeginComputeFence"));
+	FComputeFenceRHIRef ClearBeginFence = RHICmdList.CreateComputeFence(BeginComputeFenceName);
+
+	//Grab the async compute commandlist.
+	FRHIAsyncComputeCommandListImmediate& RHICmdListComputeImmediate = FRHICommandListExecutor::GetImmediateAsyncComputeCommandList();
+	{
+		SCOPED_COMPUTE_EVENTF(RHICmdListComputeImmediate, ClearTranslucencyLightingVolume, TEXT("Translucency lighting volume clear compute shader. %d"), GTranslucencyLightingVolumeDim);
+
+		RHICmdListComputeImmediate.TransitionResources(EResourceTransitionAccess::ERWBarrier, EResourceTransitionPipeline::EGfxToCompute, VolumeUAVs, 4, ClearBeginFence);
+
+		//we must wait on the fence written from the Gfx pipe to let us know all our dependencies are ready.
+		RHICmdListComputeImmediate.WaitComputeFence(ClearBeginFence);
+
+		//standard compute setup, but on the async commandlist.
+		RHICmdListComputeImmediate.SetComputeShader(ComputeShader->GetComputeShader());
+
+		ComputeShader->SetParameters(RHICmdListComputeImmediate, VolumeUAVs, 4);
+		
+		int32 GroupsPerDim = GTranslucencyLightingVolumeDim / FClearTranslucentLightingVolumeCS::CLEAR_BLOCK_SIZE;
+		DispatchComputeShader(RHICmdListComputeImmediate, ComputeShader, GroupsPerDim, GroupsPerDim, GroupsPerDim);
+
+		ComputeShader->UnsetParameters(RHICmdListComputeImmediate);
+
+		//transition the output to readable and write the fence to allow the Gfx pipe to carry on.
+		RHICmdListComputeImmediate.TransitionResources(EResourceTransitionAccess::EReadable, EResourceTransitionPipeline::EComputeToGfx, VolumeUAVs, 4, TranslucencyLightingVolumeClearEndFence);
+	}
+
+	//immediately dispatch our async compute commands to the RHI thread to be submitted to the GPU as soon as possible.
+	//dispatch after the scope so the drawevent pop is inside the dispatch
+	FRHIAsyncComputeCommandListImmediate::ImmediateDispatch(RHICmdListComputeImmediate);
+}
 
 /** Encapsulates a pixel shader that is adding ambient cubemap to the volume. */
 class FInjectAmbientCubemapPS : public FGlobalShader
diff --git a/Engine/Source/Runtime/Renderer/Public/MaterialShader.h b/Engine/Source/Runtime/Renderer/Public/MaterialShader.h
--- a/Engine/Source/Runtime/Renderer/Public/MaterialShader.h
+++ b/Engine/Source/Runtime/Renderer/Public/MaterialShader.h
@@ -164,8 +164,8 @@ private:
 	TArray<FShaderParameter> PerFramePrevScalarExpressions;
 	TArray<FShaderParameter> PerFramePrevVectorExpressions;
 	FDeferredPixelShaderParameters DeferredParameters;
-	FShaderResourceParameter LightAttenuation;
-	FShaderResourceParameter LightAttenuationSampler;
+	FShaderResourceParameter SceneColorCopyTexture;
+	FShaderResourceParameter SceneColorCopyTextureSampler;
 
 	//Use of the eye adaptation texture here is experimental and potentially dangerous as it can introduce a feedback loop. May be removed.
 	FShaderResourceParameter EyeAdaptation;
diff --git a/Engine/Source/Runtime/ShaderCore/Private/Shader.cpp b/Engine/Source/Runtime/ShaderCore/Private/Shader.cpp
--- a/Engine/Source/Runtime/ShaderCore/Private/Shader.cpp
+++ b/Engine/Source/Runtime/ShaderCore/Private/Shader.cpp
@@ -1752,4 +1752,12 @@ void ShaderMapAppendKeyString(EShaderPlatform Platform, FString& KeyString)
 			KeyString += TEXT("_SD");
 		}
 	}
+
+	{
+		static const auto CVar = IConsoleManager::Get().FindTConsoleVariableDataInt(TEXT("r.ForwardShading"));
+		if (CVar && CVar->GetValueOnAnyThread() > 0)
+		{
+			KeyString += TEXT("_FS");
+		}
+	}
 }
diff --git a/Engine/Source/Runtime/VectorVM/Private/Tests/VectorVMTests.cpp b/Engine/Source/Runtime/VectorVM/Private/Tests/VectorVMTests.cpp
--- a/Engine/Source/Runtime/VectorVM/Private/Tests/VectorVMTests.cpp
+++ b/Engine/Source/Runtime/VectorVM/Private/Tests/VectorVMTests.cpp
@@ -61,7 +61,6 @@ bool FVectorVMTest::RunTest(const FString& Parameters)
 		InputRegisters, 3,
 		OutputRegisters, 1,
 		ConstantTable,
-		nullptr,
 		dummy,
 		VectorVM::VectorsPerChunk
 		);
diff --git a/Engine/Source/Runtime/VectorVM/Private/VectorVM.cpp b/Engine/Source/Runtime/VectorVM/Private/VectorVM.cpp
--- a/Engine/Source/Runtime/VectorVM/Private/VectorVM.cpp
+++ b/Engine/Source/Runtime/VectorVM/Private/VectorVM.cpp
@@ -2,7 +2,6 @@
 
 #include "VectorVMPrivate.h"
 #include "CurveVector.h"
-#include "VectorVMDataObject.h"
 #include "ModuleManager.h"
 
 IMPLEMENT_MODULE(FDefaultModuleImpl, VectorVM);
@@ -17,10 +16,6 @@ DEFINE_LOG_CATEGORY_STATIC(LogVectorVM, All, All);
 #define OP1_CONST (1 << 1)
 #define OP2_CONST (1 << 2)
 #define OP3_CONST (1 << 3)
-#define OP0_DATAOBJ (1 << 4)
-#define OP1_DATAOBJ (1 << 5)
-#define OP2_DATAOBJ (1 << 6)
-#define OP3_DATAOBJ (1 << 7)
 
 #define SRCOP_RRRR (OP_REGISTER | OP_REGISTER | OP_REGISTER | OP_REGISTER)
 #define SRCOP_RRRC (OP_REGISTER | OP_REGISTER | OP_REGISTER | OP0_CONST)
@@ -39,56 +34,14 @@ DEFINE_LOG_CATEGORY_STATIC(LogVectorVM, All, All);
 #define SRCOP_CCCR (OP3_CONST | OP2_CONST | OP1_CONST | OP_REGISTER)
 #define SRCOP_CCCC (OP3_CONST | OP2_CONST | OP1_CONST | OP0_CONST)
 
-#define SRCOP_RRRB (OP_REGISTER | OP_REGISTER | OP_REGISTER | OP0_DATAOBJ)
-#define SRCOP_RRBR (OP_REGISTER | OP_REGISTER | OP1_DATAOBJ | OP_REGISTER)
-#define SRCOP_RRBB (OP_REGISTER | OP_REGISTER | OP1_DATAOBJ | OP0_DATAOBJ)
-
-#define SRCOP_RRCB (OP_REGISTER | OP_REGISTER | OP1_CONST | OP0_DATAOBJ)
-
 uint8 VectorVM::CreateSrcOperandMask(EVectorVMOperandLocation Type1, EVectorVMOperandLocation Type2, EVectorVMOperandLocation Type3, EVectorVMOperandLocation Type4)
 {
 	return	(Type1 == EVectorVMOperandLocation::Constant ? OP0_CONST : OP_REGISTER) |
 		(Type2 == EVectorVMOperandLocation::Constant ? OP1_CONST : OP_REGISTER) |
 		(Type3 == EVectorVMOperandLocation::Constant ? OP2_CONST : OP_REGISTER) |
-		(Type4 == EVectorVMOperandLocation::Constant ? OP3_CONST : OP_REGISTER) |
-		(Type1 == EVectorVMOperandLocation::DataObjConstant ? OP0_DATAOBJ : OP_REGISTER) |
-		(Type2 == EVectorVMOperandLocation::DataObjConstant ? OP1_DATAOBJ : OP_REGISTER) |
-		(Type3 == EVectorVMOperandLocation::DataObjConstant ? OP2_DATAOBJ : OP_REGISTER) |
-		(Type4 == EVectorVMOperandLocation::DataObjConstant ? OP3_DATAOBJ : OP_REGISTER);
-}
-
-
-UNiagaraDataObject::UNiagaraDataObject(const FObjectInitializer& ObjectInitializer)
-: Super(ObjectInitializer)
-{
+		(Type4 == EVectorVMOperandLocation::Constant ? OP3_CONST : OP_REGISTER);
 }
 
-UNiagaraCurveDataObject::UNiagaraCurveDataObject(const FObjectInitializer& ObjectInitializer)
-: Super(ObjectInitializer), CurveObj(nullptr)
-{
-}
-
-FVector4 UNiagaraCurveDataObject::Sample(const FVector4& InCoords) const
-{
-	FVector Vec = CurveObj->GetVectorValue(InCoords.X);
-	return FVector4(Vec, 0.0f);
-}
-
-UNiagaraSparseVolumeDataObject::UNiagaraSparseVolumeDataObject(const FObjectInitializer& ObjectInitializer)
-: Super(ObjectInitializer)
-{
-	Size = 64;
-	NumBuckets = Size*Size*Size;
-	//Data.AddZeroed(NumBuckets);
-
-	// avoid allocation in CDO
-	if (!HasAnyFlags(RF_ClassDefaultObject))
-	{
-		Data.Init(FVector4(0.1f, 0.1f, 0.1f, 0.1f), NumBuckets);
-	}
-}
-
-//////////////////////////////////////////////////////////////////////////
 
 #if WITH_EDITOR
 TArray<FString> OpNames;
@@ -106,8 +59,6 @@ struct FVectorVMContext
 	VectorRegister* RESTRICT * RESTRICT RegisterTable;
 	/** Pointer to the constant table. */
 	FVector4 const* RESTRICT ConstantTable;
-	/** Pointer to the data object constant table. */
-	UNiagaraDataObject * RESTRICT *DataObjConstantTable;
 	/** Pointer to the shared data table. */
 	FVectorVMSharedDataView* RESTRICT SharedDataTable;
 
@@ -127,7 +78,6 @@ struct FVectorVMContext
 		const uint8* InCode,
 		VectorRegister** InRegisterTable,
 		const FVector4* InConstantTable,
-		UNiagaraDataObject** InDataObjTable,
 		FVectorVMSharedDataView* InSharedDataTable,
 		int32 InNumVectors,
 		int32 InNumInstances,
@@ -136,7 +86,6 @@ struct FVectorVMContext
 		: Code(InCode)
 		, RegisterTable(InRegisterTable)
 		, ConstantTable(InConstantTable)
-		, DataObjConstantTable(InDataObjTable)
 		, SharedDataTable(InSharedDataTable)
 		, NumVectors(InNumVectors)
 		, NumInstances(InNumInstances)
@@ -212,19 +161,6 @@ struct FConstantHandler<VectorRegister> : public FConstantHandlerBase
 
 typedef FConstantHandler<VectorRegister> FVectorConstantHandler;
 
-//////////////////////////////////////////////////////////////////////////
-struct FDataObjectConstantHandler
-{
-	int32 ConstantIndex;
-	UNiagaraDataObject *Constant;
-	FDataObjectConstantHandler(FVectorVMContext& Context)
-		: ConstantIndex(DecodeU8(Context))
-		, Constant(Context.DataObjConstantTable[ConstantIndex])
-	{}
-	VM_FORCEINLINE void Advance(){ }
-	VM_FORCEINLINE UNiagaraDataObject *Get(){ return Constant; }
-};
-
 //////////////////////////////////////////////////////////////////////////
 // Register handlers.
 // Handle reading of a register, advancing the pointer with each read.
@@ -394,47 +330,6 @@ struct TBinaryVectorKernel
 	}
 };
 
-
-
-
-/** Base class of Vector kernels with 2 operands, one of which can be a data object. */
-template <typename Kernel, typename DstHandler = FRegisterDestHandler<VectorRegister>>
-struct TBinaryVectorKernelData
-{
-	static const int32 NumInstancesPerOp = 1;
-	static void Exec(FVectorVMContext& Context)
-	{
-		uint32 SrcOpTypes = DecodeSrcOperandTypes(Context);
-		switch (SrcOpTypes)
-		{
-		case SRCOP_RRRB:	TBinaryKernel<Kernel, DstHandler, FDataObjectConstantHandler, FVectorRegisterHandler>::Exec(Context); break;
-		case SRCOP_RRCB:	TBinaryKernel<Kernel, DstHandler, FDataObjectConstantHandler, FVectorConstantHandler>::Exec(Context); break;
-		default: check(0); break;
-		};
-	}
-};
-
-
-
-/** Base class of Vector kernels with 2 operands, one of which can be a data object. */
-template <typename Kernel, typename DstHandler = FRegisterDestHandler<VectorRegister>>
-struct TTrinaryVectorKernelData
-{
-	static const int32 NumInstancesPerOp = 1;
-	static void Exec(FVectorVMContext& Context)
-	{
-		uint32 SrcOpTypes = DecodeSrcOperandTypes(Context);
-		switch (SrcOpTypes)
-		{
-		case SRCOP_RRRB:	TTrinaryKernel<Kernel, DstHandler, FDataObjectConstantHandler, FVectorRegisterHandler, FVectorRegisterHandler>::Exec(Context); break;
-		default: check(0); break;
-		};
-	}
-};
-
-
-
-
 /** Base class of Vector kernels with 3 operands. */
 template <typename Kernel, typename DstHandler = FRegisterDestHandler<VectorRegister>>
 struct TTrinaryVectorKernel
@@ -743,34 +638,6 @@ struct FVectorKernelSelect : public TTrinaryVectorKernel<FVectorKernelSelect>
 };
 
 
-struct FVectorKernelSample : public TBinaryVectorKernelData<FVectorKernelSample>
-{
-	static void FORCEINLINE DoKernel(VectorRegister* RESTRICT Dst, UNiagaraDataObject *Src0, VectorRegister Src1)
-	{
-		if (Src0)
-		{
-			float const* FloatSrc1 = reinterpret_cast<float const*>(&Src1);
-			FVector4 Tmp = Src0->Sample(FVector4(FloatSrc1[0], FloatSrc1[1], FloatSrc1[2], FloatSrc1[3]));
-			*Dst = VectorLoad(&Tmp);
-		}
-	}
-};
-
-struct FVectorKernelBufferWrite : public TTrinaryVectorKernelData<FVectorKernelBufferWrite>
-{
-	static void FORCEINLINE DoKernel(VectorRegister* RESTRICT Dst, UNiagaraDataObject *Src0, VectorRegister Src1, VectorRegister Src2)
-	{
-		if (Src0)
-		{
-			float const* FloatSrc1 = reinterpret_cast<float const*>(&Src1);	// Coords
-			float const* FloatSrc2 = reinterpret_cast<float const*>(&Src2);	// Value
-			FVector4 Tmp = Src0->Write(FVector4(FloatSrc1[0], FloatSrc1[1], FloatSrc1[2], FloatSrc1[3]), FVector4(FloatSrc2[0], FloatSrc2[1], FloatSrc2[2], FloatSrc2[3]));
-			*Dst = VectorLoad(&Tmp);
-		}
-	}
-};
-
-
 struct FVectorKernelDot : public TBinaryVectorKernel<FVectorKernelDot>
 {
 	static void VM_FORCEINLINE DoKernel(VectorRegister* RESTRICT Dst, VectorRegister Src0, VectorRegister Src1)
@@ -1313,7 +1180,6 @@ void VectorVM::Exec(
 	VectorRegister** OutputRegisters,
 	int32 NumOutputRegisters,
 	FVector4 const* ConstantTable,
-	UNiagaraDataObject* *DataObjConstTable,
 	FVectorVMSharedDataView* SharedDataTable,
 	int32 NumVectors
 	)
@@ -1344,7 +1210,7 @@ void VectorVM::Exec(
 		// Setup execution context.
 		int32 VectorsThisChunk = FMath::Min<int32>(NumVectors, VectorsPerChunk);
 		int32 InstancesThisChunk = VectorsThisChunk;
-		FVectorVMContext Context(Code, RegisterTable, ConstantTable, DataObjConstTable, SharedDataTable, VectorsThisChunk, InstancesThisChunk, VectorsPerChunk * ChunkIndex);
+		FVectorVMContext Context(Code, RegisterTable, ConstantTable, SharedDataTable, VectorsThisChunk, InstancesThisChunk, VectorsPerChunk * ChunkIndex);
 		EVectorVMOp Op = EVectorVMOp::done;
 
 		// Execute VM on all vectors in this chunk.
@@ -1404,8 +1270,6 @@ void VectorVM::Exec(
 			case EVectorVMOp::composez: FVectorKernelCompose<2, 2, 2, 2>::Exec(Context); break;
 			case EVectorVMOp::composew: FVectorKernelCompose<3, 3, 3, 3>::Exec(Context); break;
 			case EVectorVMOp::lessthan: FVectorKernelLessThan::Exec(Context); break;
-			case EVectorVMOp::sample: FVectorKernelSample::Exec(Context); break;
-			case EVectorVMOp::bufferwrite: FVectorKernelBufferWrite::Exec(Context); break;
 			case EVectorVMOp::easein: FVectorKernelEaseIn::Exec(Context); break;
 			case EVectorVMOp::easeinout: FVectorKernelEaseInOut::Exec(Context); break;
 			case EVectorVMOp::aquireshareddataindex: FKernelSharedDataGetAppendIndex::Exec(Context); break;
diff --git a/Engine/Source/Runtime/VectorVM/Public/VectorVM.h b/Engine/Source/Runtime/VectorVM/Public/VectorVM.h
--- a/Engine/Source/Runtime/VectorVM/Public/VectorVM.h
+++ b/Engine/Source/Runtime/VectorVM/Public/VectorVM.h
@@ -9,7 +9,6 @@
 #define VECTOR_WIDTH_BYTES (16)
 #define VECTOR_WIDTH_FLOATS (4)
 
-class UNiagaraDataObject;
 class FVectorVMSharedDataView;
 struct FVectorVMContext;
 
@@ -130,7 +129,6 @@ namespace VectorVM
 		VectorRegister** OutputRegisters,
 		int32 NumOutputRegisters,
 		FVector4 const* ConstantTable,
-		UNiagaraDataObject* *DataObjTable,
 		FVectorVMSharedDataView* SharedDataTable,
 		int32 NumVectors
 		);
diff --git a/Engine/Source/Runtime/VectorVM/Public/VectorVMDataObject.h b/Engine/Source/Runtime/VectorVM/Public/VectorVMDataObject.h
--- a/Engine/Source/Runtime/VectorVM/Public/VectorVMDataObject.h
+++ b/Engine/Source/Runtime/VectorVM/Public/VectorVMDataObject.h
@@ -3,77 +3,4 @@
 #pragma once	
 #include "VectorVMDataObject.generated.h"
 
-
-/* Vector VM data object; encapsulates buffers, curves and other data in its derivatives
-*  for access by VectorVM kernels;
-*/
-UCLASS(EditInlineNew, MinimalAPI)
-class UNiagaraDataObject : public UObject
-{
-	GENERATED_UCLASS_BODY()
-public:
-	virtual FVector4 Sample(const FVector4& InCoords) const { return FVector4(0.0f, 0.0f, 0.0f, 0.0f); }
-	virtual FVector4 Write(const FVector4& InCoords, const FVector4& InValue) { return FVector4(0.0f, 0.0f, 0.0f, 0.0f); }
-};
-
-/* Curve object; encapsulates a curve for the VectorVM
-*/
-UCLASS(EditInlineNew, MinimalAPI)
-class UNiagaraCurveDataObject : public UNiagaraDataObject
-{
-	GENERATED_UCLASS_BODY()
-	
-	UPROPERTY(EditAnywhere, Category="Curve")
-	class UCurveVector *CurveObj;
-public:
-	UNiagaraCurveDataObject(class UCurveVector *InCurve) : CurveObj(InCurve)
-	{
-	}
-
-	FVector4 Sample(const FVector4& InCoords) const;
-	virtual FVector4 Write(const FVector4& InCoords, const FVector4& InValue) override
-	{
-		return FVector4(1.0f, 0.0f, 0.0f, 1.0f);
-	}
-
-	UCurveVector *GetCurveObject()	{ return CurveObj; }
-	void SetCurveObject(UCurveVector *InCurve)	{ CurveObj = InCurve; }
-};
-
-
-/* Volume data object; encapsulates volumetric data for VectorVM
-*/
-UCLASS(EditInlineNew, MinimalAPI, Transient)
-class UNiagaraSparseVolumeDataObject : public UNiagaraDataObject
-{
-	GENERATED_UCLASS_BODY()
-private:
-	TArray<FVector4> Data;
-	uint64 NumBuckets;
-	int32 Size;
-public:
-	uint32 MakeHash(const FVector4& InCoords) const
-	{
-		FVector4 Coords = InCoords;
-		const int64 P1 = 73856093; // some large primes
-		const int64 P2 = 19349663;
-		const int64 P3 = 83492791;
-		Coords *= 0.1f;
-		uint64 N = (P1*static_cast<int64>(Coords.X)) ^ (P2*static_cast<int64>(Coords.Y)) ^ (P3*static_cast<int64>(Coords.Z));
-		N %= NumBuckets;
-		return FMath::Clamp(N, 0ULL, NumBuckets);
-	}
-
-	virtual FVector4 Sample(const FVector4& InCoords) const override
-	{
-		int32 Index = MakeHash(InCoords);
-		Index = FMath::Clamp(Index, 0, Data.Num() - 1);
-		return Data[Index];
-	}
-
-	virtual FVector4 Write(const FVector4& InCoords, const FVector4& InValue) override
-	{
-		Data[MakeHash(InCoords)] = InValue;
-		return InValue;
-	}
-};
+//TODO Data Objects are to be totally re-worked
\ No newline at end of file
diff --git a/Engine/Source/Runtime/VulkanRHI/Private/VulkanCommandBuffer.cpp b/Engine/Source/Runtime/VulkanRHI/Private/VulkanCommandBuffer.cpp
--- a/Engine/Source/Runtime/VulkanRHI/Private/VulkanCommandBuffer.cpp
+++ b/Engine/Source/Runtime/VulkanRHI/Private/VulkanCommandBuffer.cpp
@@ -23,9 +23,8 @@ FVulkanCmdBuffer::FVulkanCmdBuffer(FVulkanDevice* InDevice, FVulkanCommandBuffer
 	CreateCmdBufInfo.commandBufferCount = 1;
 	CreateCmdBufInfo.commandPool = CommandBufferManager->GetHandle();
 
-	VERIFYVULKANRESULT(vkAllocateCommandBuffers(Device->GetInstanceHandle(), &CreateCmdBufInfo, &CommandBufferHandle));
+	VERIFYVULKANRESULT(VulkanRHI::vkAllocateCommandBuffers(Device->GetInstanceHandle(), &CreateCmdBufInfo, &CommandBufferHandle));
 	Fence = Device->GetFenceManager().AllocateFence();
-	//FPlatformMisc::LowLevelOutputDebugStringf(TEXT("*** Create CmdBuffer %p Fence %p\n"), (void*)CommandBufferHandle, (void*)Fence->GetHandle());
 }
 
 FVulkanCmdBuffer::~FVulkanCmdBuffer()
@@ -43,7 +42,7 @@ FVulkanCmdBuffer::~FVulkanCmdBuffer()
 		FenceManager.ReleaseFence(Fence);
 	}
 
-	vkFreeCommandBuffers(Device->GetInstanceHandle(), CommandBufferManager->GetHandle(), 1, &CommandBufferHandle);
+	VulkanRHI::vkFreeCommandBuffers(Device->GetInstanceHandle(), CommandBufferManager->GetHandle(), 1, &CommandBufferHandle);
 	CommandBufferHandle = VK_NULL_HANDLE;
 }
 
@@ -62,7 +61,7 @@ void FVulkanCmdBuffer::BeginRenderPass(const FVulkanRenderTargetLayout& Layout,
 	Info.clearValueCount = Layout.GetNumAttachments();
 	Info.pClearValues = AttachmentClearValues;
 
-	vkCmdBeginRenderPass(CommandBufferHandle, &Info, VK_SUBPASS_CONTENTS_INLINE);
+	VulkanRHI::vkCmdBeginRenderPass(CommandBufferHandle, &Info, VK_SUBPASS_CONTENTS_INLINE);
 
 	State = EState::IsInsideRenderPass;
 }
@@ -76,7 +75,7 @@ void FVulkanCmdBuffer::Begin()
 	CmdBufBeginInfo.sType = VK_STRUCTURE_TYPE_COMMAND_BUFFER_BEGIN_INFO;
 	CmdBufBeginInfo.flags = VK_COMMAND_BUFFER_USAGE_ONE_TIME_SUBMIT_BIT;
 
-	VERIFYVULKANRESULT(vkBeginCommandBuffer(CommandBufferHandle, &CmdBufBeginInfo));
+	VERIFYVULKANRESULT(VulkanRHI::vkBeginCommandBuffer(CommandBufferHandle, &CmdBufBeginInfo));
 
 	State = EState::IsInsideBegin;
 }
@@ -89,8 +88,7 @@ void FVulkanCmdBuffer::RefreshFenceStatus()
 		if (FenceMgr->IsFenceSignaled(Fence))
 		{
 			State = EState::ReadyForBegin;
-			//FPlatformMisc::LowLevelOutputDebugStringf(TEXT("vkResetCommandBuffer %p\n"), (void*)CommandBufferHandle);
-			vkResetCommandBuffer(CommandBufferHandle, VK_COMMAND_BUFFER_RESET_RELEASE_RESOURCES_BIT);
+			VulkanRHI::vkResetCommandBuffer(CommandBufferHandle, VK_COMMAND_BUFFER_RESET_RELEASE_RESOURCES_BIT);
 #if VULKAN_REUSE_FENCES
 			Fence->GetOwner()->ResetFence(Fence);
 #else
@@ -122,7 +120,7 @@ FVulkanCommandBufferManager::FVulkanCommandBufferManager(FVulkanDevice* InDevice
 	CmdPoolInfo.queueFamilyIndex = Device->GetQueue()->GetFamilyIndex();
 	//#todo-rco: Should we use VK_COMMAND_POOL_CREATE_TRANSIENT_BIT?
 	CmdPoolInfo.flags = VK_COMMAND_POOL_CREATE_RESET_COMMAND_BUFFER_BIT;
-	VERIFYVULKANRESULT(vkCreateCommandPool(Device->GetInstanceHandle(), &CmdPoolInfo, nullptr, &Handle));
+	VERIFYVULKANRESULT(VulkanRHI::vkCreateCommandPool(Device->GetInstanceHandle(), &CmdPoolInfo, nullptr, &Handle));
 
 	ActiveCmdBuffer = Create();
 	ActiveCmdBuffer->Begin();
@@ -136,7 +134,7 @@ FVulkanCommandBufferManager::~FVulkanCommandBufferManager()
 		delete CmdBuffer;
 	}
 
-	vkDestroyCommandPool(Device->GetInstanceHandle(), Handle, nullptr);
+	VulkanRHI::vkDestroyCommandPool(Device->GetInstanceHandle(), Handle, nullptr);
 }
 
 FVulkanCmdBuffer* FVulkanCommandBufferManager::GetActiveCmdBuffer()
diff --git a/Engine/Source/Runtime/VulkanRHI/Private/VulkanCommandBuffer.h b/Engine/Source/Runtime/VulkanRHI/Private/VulkanCommandBuffer.h
--- a/Engine/Source/Runtime/VulkanRHI/Private/VulkanCommandBuffer.h
+++ b/Engine/Source/Runtime/VulkanRHI/Private/VulkanCommandBuffer.h
@@ -61,14 +61,14 @@ public:
 	void EndRenderPass()
 	{
 		check(IsInsideRenderPass());
-		vkCmdEndRenderPass(CommandBufferHandle);
+		VulkanRHI::vkCmdEndRenderPass(CommandBufferHandle);
 		State = EState::IsInsideBegin;
 	}
 
 	void End()
 	{
 		check(IsOutsideRenderPass());
-		VERIFYVULKANRESULT(vkEndCommandBuffer(GetHandle()));
+		VERIFYVULKANRESULT(VulkanRHI::vkEndCommandBuffer(GetHandle()));
 		State = EState::HasEnded;
 	}
 
