Fix Ocean/WEX compile

diff --git a/Engine/Source/Runtime/VectorVM/Private/VectorVM.cpp b/Engine/Source/Runtime/VectorVM/Private/VectorVM.cpp
--- a/Engine/Source/Runtime/VectorVM/Private/VectorVM.cpp
+++ b/Engine/Source/Runtime/VectorVM/Private/VectorVM.cpp
@@ -985,7 +985,7 @@ struct FVectorKernelReadInput
 {
 	static VM_FORCEINLINE void Exec(FVectorVMContext& Context)
 	{
-		constexpr int32 InstancesPerVector = sizeof(VectorRegister) / sizeof(T);
+		static const int32 InstancesPerVector = sizeof(VectorRegister) / sizeof(T);
 
 		int32 DataSetIndex = DecodeU16(Context);
 		int32 InputRegisterIdx = DecodeU16(Context);
@@ -1017,7 +1017,7 @@ struct FVectorKernelReadInputNoAdvance
 {
 	static VM_FORCEINLINE void Exec(FVectorVMContext& Context)
 	{
-		constexpr int32 InstancesPerVector = sizeof(VectorRegister) / sizeof(T);
+		static const int32 InstancesPerVector = sizeof(VectorRegister) / sizeof(T);
 
 		int32 DataSetIndex = DecodeU16(Context);
 		int32 InputRegisterIdx = DecodeU16(Context);
