Copying //UE4/Dev-Niagara to //UE4/Dev-Main (Source: //UE4/Dev-Niagara @ 4074996)

diff --git a/Engine/Source/Runtime/VectorVM/Private/VectorVM.cpp b/Engine/Source/Runtime/VectorVM/Private/VectorVM.cpp
--- a/Engine/Source/Runtime/VectorVM/Private/VectorVM.cpp
+++ b/Engine/Source/Runtime/VectorVM/Private/VectorVM.cpp
@@ -6,6 +6,8 @@
 #include "UObject/Package.h"
 #include "VectorVMPrivate.h"
 #include "Stats/Stats.h"
+#include "HAL/ConsoleManager.h"
+#include "Async/ParallelFor.h"
 
 IMPLEMENT_MODULE(FDefaultModuleImpl, VectorVM);
 
@@ -32,6 +34,22 @@ DEFINE_LOG_CATEGORY_STATIC(LogVectorVM, All, All);
 #define SRCOP_CCR (OP2_CONST | OP1_CONST | OP_REGISTER)
 #define SRCOP_CCC (OP2_CONST | OP1_CONST | OP0_CONST)
 
+static int32 GbParallelVVM = 1;
+static FAutoConsoleVariableRef CVarbParallelVVM(
+	TEXT("vm.Parallel"),
+	GbParallelVVM,
+	TEXT("If > 0 vector VM chunk level paralellism will be enabled. \n"),
+	ECVF_Default
+);
+
+static int32 GParallelVVMChunksPerBatch = 4;
+static FAutoConsoleVariableRef CVarParallelVVMChunksPerBatch(
+	TEXT("vm.ParallelChunksPerBatch"),
+	GParallelVVMChunksPerBatch,
+	TEXT("Number of chunks to process per task when running in parallel. \n"),
+	ECVF_Default
+);
+
 //////////////////////////////////////////////////////////////////////////
 
 FVectorVMContext::FVectorVMContext()
@@ -44,6 +62,7 @@ FVectorVMContext::FVectorVMContext()
 	, UserPtrTable(nullptr)
 	, NumInstances(0)
 	, StartInstance(0)
+	, DataSetMetaTable(nullptr)
 #if STATS
 	, StatScopes(nullptr)
 #endif
@@ -55,6 +74,8 @@ FVectorVMContext::FVectorVMContext()
 	{
 		RegisterTable[i] = TempRegTable.GetData() + TempRegisterSize * i;
 	}
+
+	RandStream.GenerateNewSeed();
 }
 
 void FVectorVMContext::PrepareForExec(
@@ -67,7 +88,8 @@ void FVectorVMContext::PrepareForExec(
 	int32 *InDataSetOffsetTable,
 	int32 InNumSecondaryDatasets,
 	FVMExternalFunction* InExternalFunctionTable,
-	void** InUserPtrTable
+	void** InUserPtrTable,
+	FDataSetMeta*RESTRICT InDataSetMetaTable
 #if STATS
 	, const TArray<TStatId>* InStatScopes
 #endif
@@ -94,6 +116,8 @@ void FVectorVMContext::PrepareForExec(
 	{
 		RegisterTable[VectorVM::NumTempRegisters + VectorVM::MaxInputRegisters + i] = OutputRegisters[i];
 	}
+
+	DataSetMetaTable = InDataSetMetaTable;
 }
 
 //////////////////////////////////////////////////////////////////////////
@@ -118,8 +142,7 @@ struct TUnaryKernelHandler
 		int32 LoopInstances = Align(Context.NumInstances, NumInstancesPerOp) / NumInstancesPerOp;
 		for (int32 i = 0; i < LoopInstances; ++i)
 		{
-			Kernel::DoKernel(Dst.GetDest(), Arg0.Get());
-			Dst.Advance();	Arg0.Advance();
+			Kernel::DoKernel(Context, Dst.GetDestAndAdvance(), Arg0.GetAndAdvance());
 		}
 	}
 };
@@ -137,8 +160,7 @@ struct TBinaryKernelHandler
 		int32 LoopInstances = Align(Context.NumInstances, NumInstancesPerOp) / NumInstancesPerOp;
 		for (int32 i = 0; i < LoopInstances; ++i)
 		{
-			Kernel::DoKernel(Dst.GetDest(), Arg0.Get(), Arg1.Get());
-			Dst.Advance(); Arg0.Advance(); Arg1.Advance();
+			Kernel::DoKernel(Context, Dst.GetDestAndAdvance(), Arg0.GetAndAdvance(), Arg1.GetAndAdvance());
 		}
 	}
 };
@@ -157,8 +179,7 @@ struct TTrinaryKernelHandler
 		int32 LoopInstances = Align(Context.NumInstances, NumInstancesPerOp) / NumInstancesPerOp;
 		for (int32 i = 0; i < LoopInstances; ++i)
 		{
-			Kernel::DoKernel(Dst.GetDest(), Arg0.Get(), Arg1.Get(), Arg2.Get());
-			Dst.Advance(); Arg0.Advance(); Arg1.Advance(); Arg2.Advance();
+			Kernel::DoKernel(Context, Dst.GetDestAndAdvance(), Arg0.GetAndAdvance(), Arg1.GetAndAdvance(), Arg2.GetAndAdvance());
 		}
 	}
 };
@@ -178,8 +199,7 @@ struct TTrinaryOutputKernelHandler
 		int32 LoopInstances = Align(Context.NumInstances, NumInstancesPerOp) / NumInstancesPerOp;
 		for (int32 i = 0; i < LoopInstances; ++i)
 		{
-			Kernel::DoKernel(Dst.GetDest(), Arg0.Get(), Arg1.Get(), Arg2.Get());
-			Dst.Advance(); Arg0.Advance(); Arg1.Advance(); Arg2.Advance();
+			Kernel::DoKernel(Context, Dst.GetDestAndAdvance(), Arg0.GetAndAdvance(), Arg1.GetAndAdvance(), Arg2.GetAndAdvance());
 		}
 	}
 };
@@ -269,7 +289,7 @@ struct TTrinaryVectorIntKernel : public TTrinaryKernel<Kernel, FRegisterDestHand
 
 struct FVectorKernelAdd : public TBinaryVectorKernel<FVectorKernelAdd>
 {
-	static void VM_FORCEINLINE DoKernel(VectorRegister* RESTRICT Dst,VectorRegister Src0,VectorRegister Src1)
+	static void VM_FORCEINLINE DoKernel(FVectorVMContext& Context, VectorRegister* RESTRICT Dst,VectorRegister Src0,VectorRegister Src1)
 	{
 		*Dst = VectorAdd(Src0, Src1);
 	}
@@ -277,7 +297,7 @@ struct FVectorKernelAdd : public TBinaryVectorKernel<FVectorKernelAdd>
 
 struct FVectorKernelSub : public TBinaryVectorKernel<FVectorKernelSub>
 {
-	static void VM_FORCEINLINE DoKernel(VectorRegister* RESTRICT Dst,VectorRegister Src0,VectorRegister Src1)
+	static void VM_FORCEINLINE DoKernel(FVectorVMContext& Context, VectorRegister* RESTRICT Dst,VectorRegister Src0,VectorRegister Src1)
 	{
 		*Dst = VectorSubtract(Src0, Src1);
 	}
@@ -285,7 +305,7 @@ struct FVectorKernelSub : public TBinaryVectorKernel<FVectorKernelSub>
 
 struct FVectorKernelMul : public TBinaryVectorKernel<FVectorKernelMul>
 {
-	static void VM_FORCEINLINE DoKernel(VectorRegister* RESTRICT Dst,VectorRegister Src0,VectorRegister Src1)
+	static void VM_FORCEINLINE DoKernel(FVectorVMContext& Context, VectorRegister* RESTRICT Dst,VectorRegister Src0,VectorRegister Src1)
 	{
 		*Dst = VectorMultiply(Src0, Src1);
 	}
@@ -293,7 +313,7 @@ struct FVectorKernelMul : public TBinaryVectorKernel<FVectorKernelMul>
 
 struct FVectorKernelDiv : public TBinaryVectorKernel<FVectorKernelDiv>
 {
-	static void VM_FORCEINLINE DoKernel(VectorRegister* RESTRICT Dst, VectorRegister Src0, VectorRegister Src1)
+	static void VM_FORCEINLINE DoKernel(FVectorVMContext& Context, VectorRegister* RESTRICT Dst, VectorRegister Src0, VectorRegister Src1)
 	{
 		*Dst = VectorDivide(Src0, Src1);
 	}
@@ -302,7 +322,7 @@ struct FVectorKernelDiv : public TBinaryVectorKernel<FVectorKernelDiv>
 
 struct FVectorKernelMad : public TTrinaryVectorKernel<FVectorKernelMad>
 {
-	static void VM_FORCEINLINE DoKernel(VectorRegister* RESTRICT Dst,VectorRegister Src0,VectorRegister Src1,VectorRegister Src2)
+	static void VM_FORCEINLINE DoKernel(FVectorVMContext& Context, VectorRegister* RESTRICT Dst,VectorRegister Src0,VectorRegister Src1,VectorRegister Src2)
 	{
 		*Dst = VectorMultiplyAdd(Src0, Src1, Src2);
 	}
@@ -310,7 +330,7 @@ struct FVectorKernelMad : public TTrinaryVectorKernel<FVectorKernelMad>
 
 struct FVectorKernelLerp : public TTrinaryVectorKernel<FVectorKernelLerp>
 {
-	static void VM_FORCEINLINE DoKernel(VectorRegister* RESTRICT Dst,VectorRegister Src0,VectorRegister Src1,VectorRegister Src2)
+	static void VM_FORCEINLINE DoKernel(FVectorVMContext& Context, VectorRegister* RESTRICT Dst,VectorRegister Src0,VectorRegister Src1,VectorRegister Src2)
 	{
 		const VectorRegister OneMinusAlpha = VectorSubtract(GlobalVectorConstants::FloatOne, Src2);
 		const VectorRegister Tmp = VectorMultiply(Src0, OneMinusAlpha);
@@ -320,7 +340,7 @@ struct FVectorKernelLerp : public TTrinaryVectorKernel<FVectorKernelLerp>
 
 struct FVectorKernelRcp : public TUnaryVectorKernel<FVectorKernelRcp>
 {
-	static void VM_FORCEINLINE DoKernel(VectorRegister* RESTRICT Dst,VectorRegister Src0)
+	static void VM_FORCEINLINE DoKernel(FVectorVMContext& Context, VectorRegister* RESTRICT Dst,VectorRegister Src0)
 	{
 		*Dst = VectorReciprocal(Src0);
 	}
@@ -328,7 +348,7 @@ struct FVectorKernelRcp : public TUnaryVectorKernel<FVectorKernelRcp>
 
 struct FVectorKernelRsq : public TUnaryVectorKernel<FVectorKernelRsq>
 {
-	static void VM_FORCEINLINE DoKernel(VectorRegister* RESTRICT Dst,VectorRegister Src0)
+	static void VM_FORCEINLINE DoKernel(FVectorVMContext& Context, VectorRegister* RESTRICT Dst,VectorRegister Src0)
 	{
 		*Dst = VectorReciprocalSqrt(Src0);
 	}
@@ -336,7 +356,7 @@ struct FVectorKernelRsq : public TUnaryVectorKernel<FVectorKernelRsq>
 
 struct FVectorKernelSqrt : public TUnaryVectorKernel<FVectorKernelSqrt>
 {
-	static void VM_FORCEINLINE DoKernel(VectorRegister* RESTRICT Dst,VectorRegister Src0)
+	static void VM_FORCEINLINE DoKernel(FVectorVMContext& Context, VectorRegister* RESTRICT Dst,VectorRegister Src0)
 	{
 		// TODO: Need a SIMD sqrt!
 		*Dst = VectorReciprocal(VectorReciprocalSqrt(Src0));
@@ -345,7 +365,7 @@ struct FVectorKernelSqrt : public TUnaryVectorKernel<FVectorKernelSqrt>
 
 struct FVectorKernelNeg : public TUnaryVectorKernel<FVectorKernelNeg>
 {
-	static void VM_FORCEINLINE DoKernel(VectorRegister* RESTRICT Dst,VectorRegister Src0)
+	static void VM_FORCEINLINE DoKernel(FVectorVMContext& Context, VectorRegister* RESTRICT Dst,VectorRegister Src0)
 	{
 		*Dst = VectorNegate(Src0);
 	}
@@ -353,7 +373,7 @@ struct FVectorKernelNeg : public TUnaryVectorKernel<FVectorKernelNeg>
 
 struct FVectorKernelAbs : public TUnaryVectorKernel<FVectorKernelAbs>
 {
-	static void VM_FORCEINLINE DoKernel(VectorRegister* RESTRICT Dst,VectorRegister Src0)
+	static void VM_FORCEINLINE DoKernel(FVectorVMContext& Context, VectorRegister* RESTRICT Dst,VectorRegister Src0)
 	{
 		*Dst = VectorAbs(Src0);
 	}
@@ -361,7 +381,7 @@ struct FVectorKernelAbs : public TUnaryVectorKernel<FVectorKernelAbs>
 
 struct FVectorKernelExp : public TUnaryVectorKernel<FVectorKernelExp>
 {
-	static void VM_FORCEINLINE DoKernel(VectorRegister* RESTRICT Dst, VectorRegister Src0)
+	static void VM_FORCEINLINE DoKernel(FVectorVMContext& Context, VectorRegister* RESTRICT Dst, VectorRegister Src0)
 	{
 		*Dst = VectorExp(Src0);
 	}
@@ -369,7 +389,7 @@ struct FVectorKernelExp : public TUnaryVectorKernel<FVectorKernelExp>
 
 struct FVectorKernelExp2 : public TUnaryVectorKernel<FVectorKernelExp2>
 {
-	static void VM_FORCEINLINE DoKernel(VectorRegister* RESTRICT Dst, VectorRegister Src0)
+	static void VM_FORCEINLINE DoKernel(FVectorVMContext& Context, VectorRegister* RESTRICT Dst, VectorRegister Src0)
 	{
 		*Dst = VectorExp2(Src0);
 	}
@@ -377,7 +397,7 @@ struct FVectorKernelExp2 : public TUnaryVectorKernel<FVectorKernelExp2>
 
 struct FVectorKernelLog : public TUnaryVectorKernel<FVectorKernelLog>
 {
-	static void VM_FORCEINLINE DoKernel(VectorRegister* RESTRICT Dst, VectorRegister Src0)
+	static void VM_FORCEINLINE DoKernel(FVectorVMContext& Context, VectorRegister* RESTRICT Dst, VectorRegister Src0)
 	{
 		*Dst = VectorLog(Src0);
 	}
@@ -385,7 +405,7 @@ struct FVectorKernelLog : public TUnaryVectorKernel<FVectorKernelLog>
 
 struct FVectorKernelLog2 : public TUnaryVectorKernel<FVectorKernelLog2>
 {
-	static void VM_FORCEINLINE DoKernel(VectorRegister* RESTRICT Dst, VectorRegister Src0)
+	static void VM_FORCEINLINE DoKernel(FVectorVMContext& Context, VectorRegister* RESTRICT Dst, VectorRegister Src0)
 	{
 		*Dst = VectorLog2(Src0);
 	}
@@ -393,7 +413,7 @@ struct FVectorKernelLog2 : public TUnaryVectorKernel<FVectorKernelLog2>
 
 struct FVectorKernelClamp : public TTrinaryVectorKernel<FVectorKernelClamp>
 {
-	static void VM_FORCEINLINE DoKernel(VectorRegister* RESTRICT Dst,VectorRegister Src0,VectorRegister Src1,VectorRegister Src2)
+	static void VM_FORCEINLINE DoKernel(FVectorVMContext& Context, VectorRegister* RESTRICT Dst,VectorRegister Src0,VectorRegister Src1,VectorRegister Src2)
 	{
 		const VectorRegister Tmp = VectorMax(Src0, Src1);
 		*Dst = VectorMin(Tmp, Src2);
@@ -402,63 +422,63 @@ struct FVectorKernelClamp : public TTrinaryVectorKernel<FVectorKernelClamp>
 
 struct FVectorKernelSin : public TUnaryVectorKernel<FVectorKernelSin>
 {
-	static void VM_FORCEINLINE DoKernel(VectorRegister* RESTRICT Dst, VectorRegister Src0)
+	static void VM_FORCEINLINE DoKernel(FVectorVMContext& Context, VectorRegister* RESTRICT Dst, VectorRegister Src0)
 	{
-		*Dst = VectorSin(VectorMultiply(Src0, GlobalVectorConstants::TwoPi));
+		*Dst = VectorSin(Src0);
 	}
 };
 
 struct FVectorKernelCos : public TUnaryVectorKernel<FVectorKernelCos>
 {
-	static void VM_FORCEINLINE DoKernel(VectorRegister* RESTRICT Dst, VectorRegister Src0)
+	static void VM_FORCEINLINE DoKernel(FVectorVMContext& Context, VectorRegister* RESTRICT Dst, VectorRegister Src0)
  	{
-		*Dst = VectorCos(VectorMultiply(Src0, GlobalVectorConstants::TwoPi));
+		*Dst = VectorCos(Src0);
 	}
 };
 
 struct FVectorKernelTan : public TUnaryVectorKernel<FVectorKernelTan>
 {
-	static void VM_FORCEINLINE DoKernel(VectorRegister* RESTRICT Dst, VectorRegister Src0)
+	static void VM_FORCEINLINE DoKernel(FVectorVMContext& Context, VectorRegister* RESTRICT Dst, VectorRegister Src0)
 	{
-		*Dst = VectorTan(VectorMultiply(Src0, GlobalVectorConstants::TwoPi));
+		*Dst = VectorTan(Src0);
 	}
 };
 
 struct FVectorKernelASin : public TUnaryVectorKernel<FVectorKernelASin>
 {
-	static void VM_FORCEINLINE DoKernel(VectorRegister* RESTRICT Dst, VectorRegister Src0)
+	static void VM_FORCEINLINE DoKernel(FVectorVMContext& Context, VectorRegister* RESTRICT Dst, VectorRegister Src0)
 	{
-		*Dst = VectorMultiply(VectorASin(Src0), GlobalVectorConstants::OneOverTwoPi);
+		*Dst = VectorASin(Src0);
 	}
 };
 
 struct FVectorKernelACos : public TUnaryVectorKernel<FVectorKernelACos>
 {
-	static void VM_FORCEINLINE DoKernel(VectorRegister* RESTRICT Dst, VectorRegister Src0)
+	static void VM_FORCEINLINE DoKernel(FVectorVMContext& Context, VectorRegister* RESTRICT Dst, VectorRegister Src0)
 	{
-		*Dst = VectorMultiply(VectorACos(Src0), GlobalVectorConstants::OneOverTwoPi);
+		*Dst = VectorACos(Src0);
 	}
 };
 
 struct FVectorKernelATan : public TUnaryVectorKernel<FVectorKernelATan>
 {
-	static void VM_FORCEINLINE DoKernel(VectorRegister* RESTRICT Dst, VectorRegister Src0)
+	static void VM_FORCEINLINE DoKernel(FVectorVMContext& Context, VectorRegister* RESTRICT Dst, VectorRegister Src0)
 	{
-		*Dst = VectorMultiply(VectorATan(Src0), GlobalVectorConstants::OneOverTwoPi);
+		*Dst = VectorATan(Src0);
 	}
 };
 
 struct FVectorKernelATan2 : public TBinaryVectorKernel<FVectorKernelATan2>
 {
-	static void VM_FORCEINLINE DoKernel(VectorRegister* RESTRICT Dst, VectorRegister Src0, VectorRegister Src1)
+	static void VM_FORCEINLINE DoKernel(FVectorVMContext& Context, VectorRegister* RESTRICT Dst, VectorRegister Src0, VectorRegister Src1)
 	{
-		*Dst = VectorMultiply(VectorATan2(Src0, Src1), GlobalVectorConstants::OneOverTwoPi);
+		*Dst = VectorATan2(Src0, Src1);
 	}
 };
 
 struct FVectorKernelCeil : public TUnaryVectorKernel<FVectorKernelCeil>
 {
-	static void VM_FORCEINLINE DoKernel(VectorRegister* RESTRICT Dst, VectorRegister Src0)
+	static void VM_FORCEINLINE DoKernel(FVectorVMContext& Context, VectorRegister* RESTRICT Dst, VectorRegister Src0)
 	{
 		*Dst = VectorCeil(Src0);
 	}
@@ -466,7 +486,7 @@ struct FVectorKernelCeil : public TUnaryVectorKernel<FVectorKernelCeil>
 
 struct FVectorKernelFloor : public TUnaryVectorKernel<FVectorKernelFloor>
 {
-	static void VM_FORCEINLINE DoKernel(VectorRegister* RESTRICT Dst, VectorRegister Src0)
+	static void VM_FORCEINLINE DoKernel(FVectorVMContext& Context, VectorRegister* RESTRICT Dst, VectorRegister Src0)
 	{
 		*Dst = VectorFloor(Src0);
 	}
@@ -474,7 +494,7 @@ struct FVectorKernelFloor : public TUnaryVectorKernel<FVectorKernelFloor>
 
 struct FVectorKernelRound : public TUnaryVectorKernel<FVectorKernelRound>
 {
-	static void VM_FORCEINLINE DoKernel(VectorRegister* RESTRICT Dst, VectorRegister Src0)
+	static void VM_FORCEINLINE DoKernel(FVectorVMContext& Context, VectorRegister* RESTRICT Dst, VectorRegister Src0)
 	{
 		//TODO: >SSE4 has direct ops for this.		
 		VectorRegister Trunc = VectorTruncate(Src0);
@@ -484,7 +504,7 @@ struct FVectorKernelRound : public TUnaryVectorKernel<FVectorKernelRound>
 
 struct FVectorKernelMod : public TBinaryVectorKernel<FVectorKernelMod>
 {
-	static void VM_FORCEINLINE DoKernel(VectorRegister* RESTRICT Dst, VectorRegister Src0, VectorRegister Src1)
+	static void VM_FORCEINLINE DoKernel(FVectorVMContext& Context, VectorRegister* RESTRICT Dst, VectorRegister Src0, VectorRegister Src1)
 	{
 		*Dst = VectorMod(Src0, Src1);
 	}
@@ -492,7 +512,7 @@ struct FVectorKernelMod : public TBinaryVectorKernel<FVectorKernelMod>
 
 struct FVectorKernelFrac : public TUnaryVectorKernel<FVectorKernelFrac>
 {
-	static void VM_FORCEINLINE DoKernel(VectorRegister* RESTRICT Dst, VectorRegister Src0)
+	static void VM_FORCEINLINE DoKernel(FVectorVMContext& Context, VectorRegister* RESTRICT Dst, VectorRegister Src0)
 	{
 		*Dst = VectorFractional(Src0);
 	}
@@ -500,7 +520,7 @@ struct FVectorKernelFrac : public TUnaryVectorKernel<FVectorKernelFrac>
 
 struct FVectorKernelTrunc : public TUnaryVectorKernel<FVectorKernelTrunc>
 {
-	static void VM_FORCEINLINE DoKernel(VectorRegister* RESTRICT Dst, VectorRegister Src0)
+	static void VM_FORCEINLINE DoKernel(FVectorVMContext& Context, VectorRegister* RESTRICT Dst, VectorRegister Src0)
 	{
 		*Dst = VectorTruncate(Src0);
 	}
@@ -508,7 +528,7 @@ struct FVectorKernelTrunc : public TUnaryVectorKernel<FVectorKernelTrunc>
 
 struct FVectorKernelCompareLT : public TBinaryVectorKernel<FVectorKernelCompareLT>
 {
-	static void VM_FORCEINLINE DoKernel(VectorRegister* RESTRICT Dst, VectorRegister Src0, VectorRegister Src1)
+	static void VM_FORCEINLINE DoKernel(FVectorVMContext& Context, VectorRegister* RESTRICT Dst, VectorRegister Src0, VectorRegister Src1)
 	{
 		*Dst = VectorCompareLT(Src0, Src1);
 	}
@@ -516,7 +536,7 @@ struct FVectorKernelCompareLT : public TBinaryVectorKernel<FVectorKernelCompareL
 
 struct FVectorKernelCompareLE : public TBinaryVectorKernel<FVectorKernelCompareLE>
 {
-	static void VM_FORCEINLINE DoKernel(VectorRegister* RESTRICT Dst, VectorRegister Src0, VectorRegister Src1)
+	static void VM_FORCEINLINE DoKernel(FVectorVMContext& Context, VectorRegister* RESTRICT Dst, VectorRegister Src0, VectorRegister Src1)
 	{
 		*Dst = VectorCompareLE(Src0, Src1);
 	}
@@ -524,7 +544,7 @@ struct FVectorKernelCompareLE : public TBinaryVectorKernel<FVectorKernelCompareL
 
 struct FVectorKernelCompareGT : public TBinaryVectorKernel<FVectorKernelCompareGT>
 {
-	static void VM_FORCEINLINE DoKernel(VectorRegister* RESTRICT Dst, VectorRegister Src0, VectorRegister Src1)
+	static void VM_FORCEINLINE DoKernel(FVectorVMContext& Context, VectorRegister* RESTRICT Dst, VectorRegister Src0, VectorRegister Src1)
 	{
 		*Dst = VectorCompareGT(Src0, Src1);
 	}
@@ -532,7 +552,7 @@ struct FVectorKernelCompareGT : public TBinaryVectorKernel<FVectorKernelCompareG
 
 struct FVectorKernelCompareGE : public TBinaryVectorKernel<FVectorKernelCompareGE>
 {
-	static void VM_FORCEINLINE DoKernel(VectorRegister* RESTRICT Dst, VectorRegister Src0, VectorRegister Src1)
+	static void VM_FORCEINLINE DoKernel(FVectorVMContext& Context, VectorRegister* RESTRICT Dst, VectorRegister Src0, VectorRegister Src1)
 	{
 		*Dst = VectorCompareGE(Src0, Src1);
 	}
@@ -540,7 +560,7 @@ struct FVectorKernelCompareGE : public TBinaryVectorKernel<FVectorKernelCompareG
 
 struct FVectorKernelCompareEQ : public TBinaryVectorKernel<FVectorKernelCompareEQ>
 {
-	static void VM_FORCEINLINE DoKernel(VectorRegister* RESTRICT Dst, VectorRegister Src0, VectorRegister Src1)
+	static void VM_FORCEINLINE DoKernel(FVectorVMContext& Context, VectorRegister* RESTRICT Dst, VectorRegister Src0, VectorRegister Src1)
 	{
 		*Dst = VectorCompareEQ(Src0, Src1);
 	}
@@ -548,7 +568,7 @@ struct FVectorKernelCompareEQ : public TBinaryVectorKernel<FVectorKernelCompareE
 
 struct FVectorKernelCompareNEQ : public TBinaryVectorKernel<FVectorKernelCompareNEQ>
 {
-	static void VM_FORCEINLINE DoKernel(VectorRegister* RESTRICT Dst, VectorRegister Src0, VectorRegister Src1)
+	static void VM_FORCEINLINE DoKernel(FVectorVMContext& Context, VectorRegister* RESTRICT Dst, VectorRegister Src0, VectorRegister Src1)
 	{		
 		*Dst = VectorCompareNE(Src0, Src1);
 	}
@@ -556,7 +576,7 @@ struct FVectorKernelCompareNEQ : public TBinaryVectorKernel<FVectorKernelCompare
 
 struct FVectorKernelSelect : public TTrinaryVectorKernel<FVectorKernelSelect>
 {
-	static void VM_FORCEINLINE DoKernel(VectorRegister* RESTRICT Dst, VectorRegister Mask, VectorRegister A, VectorRegister B)
+	static void VM_FORCEINLINE DoKernel(FVectorVMContext& Context, VectorRegister* RESTRICT Dst, VectorRegister Mask, VectorRegister A, VectorRegister B)
 	{
 		*Dst = VectorSelect(Mask, A, B);
 	}
@@ -574,8 +594,7 @@ struct FVectorKernelExecutionIndex
 		int32 Loops = Align(Context.NumInstances, VECTOR_WIDTH_FLOATS) / VECTOR_WIDTH_FLOATS;
 		for (int32 i = 0; i < Loops; ++i)
 		{
-			*Dest.GetDest() = Index;
-			Dest.Advance();
+			*Dest.GetDestAndAdvance() = Index;
 			Index = VectorIntAdd(Index, VectorStride);
 		}
 	}
@@ -587,8 +606,8 @@ struct FVectorKernelEnterStatScope
 	{
 		FConstantHandler<int32> ScopeIdx(Context);
 #if STATS
-		int32 CounterIdx = Context.StatCounterStack.AddDefaulted(1);
-		Context.StatCounterStack[CounterIdx].Start((*Context.StatScopes)[ScopeIdx.Get()]);
+		//int32 CounterIdx = Context.StatCounterStack.AddDefaulted(1);
+		//Context.StatCounterStack[CounterIdx].Start((*Context.StatScopes)[ScopeIdx.Get()]);
 #endif
 	}
 };
@@ -598,22 +617,22 @@ struct FVectorKernelExitStatScope
 	static VM_FORCEINLINE void Exec(FVectorVMContext& Context)
 	{
 #if STATS
-		Context.StatCounterStack.Last().Stop();
-		Context.StatCounterStack.Pop(false);
+		//Context.StatCounterStack.Last().Stop();
+		//Context.StatCounterStack.Pop(false);
 #endif
 	}
 };
 
 struct FVectorKernelRandom : public TUnaryVectorKernel<FVectorKernelRandom>
 {
-	static void VM_FORCEINLINE DoKernel(VectorRegister* RESTRICT Dst, VectorRegister Src0)
+	static void VM_FORCEINLINE DoKernel(FVectorVMContext& Context, VectorRegister* RESTRICT Dst, VectorRegister Src0)
 	{
 		const float rm = RAND_MAX;
 		//EEK!. Improve this. Implement GPU style seeded rand instead of this.
-		VectorRegister Result = MakeVectorRegister(static_cast<float>(FMath::Rand()) / rm,
-			static_cast<float>(FMath::Rand()) / rm,
-			static_cast<float>(FMath::Rand()) / rm,
-			static_cast<float>(FMath::Rand()) / rm);
+		VectorRegister Result = MakeVectorRegister(Context.RandStream.GetFraction(),
+			Context.RandStream.GetFraction(),
+			Context.RandStream.GetFraction(),
+			Context.RandStream.GetFraction());
 		*Dst = VectorMultiply(Result, Src0);
 	}
 };
@@ -621,13 +640,13 @@ struct FVectorKernelRandom : public TUnaryVectorKernel<FVectorKernelRandom>
 /* gaussian distribution random number (not working yet) */
 struct FVectorKernelRandomGauss : public TBinaryVectorKernel<FVectorKernelRandomGauss>
 {
-	static void VM_FORCEINLINE DoKernel(VectorRegister* RESTRICT Dst, VectorRegister Src0, VectorRegister Src1)
+	static void VM_FORCEINLINE DoKernel(FVectorVMContext& Context, VectorRegister* RESTRICT Dst, VectorRegister Src0, VectorRegister Src1)
 	{
 		const float rm = RAND_MAX;
-		VectorRegister Result = MakeVectorRegister(static_cast<float>(FMath::Rand()) / rm,
-			static_cast<float>(FMath::Rand()) / rm,
-			static_cast<float>(FMath::Rand()) / rm,
-			static_cast<float>(FMath::Rand()) / rm);
+		VectorRegister Result = MakeVectorRegister(Context.RandStream.GetFraction(),
+			Context.RandStream.GetFraction(),
+			Context.RandStream.GetFraction(),
+			Context.RandStream.GetFraction());
 
 		Result = VectorSubtract(Result, MakeVectorRegister(0.5f, 0.5f, 0.5f, 0.5f));
 		Result = VectorMultiply(MakeVectorRegister(3.0f, 3.0f, 3.0f, 3.0f), Result);
@@ -650,7 +669,7 @@ struct FVectorKernelRandomGauss : public TBinaryVectorKernel<FVectorKernelRandom
 
 struct FVectorKernelMin : public TBinaryVectorKernel<FVectorKernelMin>
 {
-	static void VM_FORCEINLINE DoKernel(VectorRegister* RESTRICT Dst,VectorRegister Src0,VectorRegister Src1)
+	static void VM_FORCEINLINE DoKernel(FVectorVMContext& Context, VectorRegister* RESTRICT Dst,VectorRegister Src0,VectorRegister Src1)
 	{
 		*Dst = VectorMin(Src0, Src1);
 	}
@@ -658,7 +677,7 @@ struct FVectorKernelMin : public TBinaryVectorKernel<FVectorKernelMin>
 
 struct FVectorKernelMax : public TBinaryVectorKernel<FVectorKernelMax>
 {
-	static void VM_FORCEINLINE DoKernel(VectorRegister* RESTRICT Dst,VectorRegister Src0,VectorRegister Src1)
+	static void VM_FORCEINLINE DoKernel(FVectorVMContext& Context, VectorRegister* RESTRICT Dst,VectorRegister Src0,VectorRegister Src1)
 	{
 		*Dst = VectorMax(Src0, Src1);
 	}
@@ -666,7 +685,7 @@ struct FVectorKernelMax : public TBinaryVectorKernel<FVectorKernelMax>
 
 struct FVectorKernelPow : public TBinaryVectorKernel<FVectorKernelPow>
 {
-	static void VM_FORCEINLINE DoKernel(VectorRegister* RESTRICT Dst,VectorRegister Src0,VectorRegister Src1)
+	static void VM_FORCEINLINE DoKernel(FVectorVMContext& Context, VectorRegister* RESTRICT Dst,VectorRegister Src0,VectorRegister Src1)
 	{
 		*Dst = VectorPow(Src0, Src1);
 	}
@@ -674,7 +693,7 @@ struct FVectorKernelPow : public TBinaryVectorKernel<FVectorKernelPow>
 
 struct FVectorKernelSign : public TUnaryVectorKernel<FVectorKernelSign>
 {
-	static void VM_FORCEINLINE DoKernel(VectorRegister* RESTRICT Dst, VectorRegister Src0)
+	static void VM_FORCEINLINE DoKernel(FVectorVMContext& Context, VectorRegister* RESTRICT Dst, VectorRegister Src0)
 	{
 		*Dst = VectorSign(Src0);
 	}
@@ -682,7 +701,7 @@ struct FVectorKernelSign : public TUnaryVectorKernel<FVectorKernelSign>
 
 struct FVectorKernelStep : public TUnaryVectorKernel<FVectorKernelStep>
 {
-	static void VM_FORCEINLINE DoKernel(VectorRegister* RESTRICT Dst, VectorRegister Src0)
+	static void VM_FORCEINLINE DoKernel(FVectorVMContext& Context, VectorRegister* RESTRICT Dst, VectorRegister Src0)
 	{
 		*Dst = VectorStep(Src0);
 	}
@@ -739,7 +758,7 @@ namespace VectorVMNoise
 
 	struct FScalarKernelNoise3D_iNoise : TTrinaryScalarKernel<FScalarKernelNoise3D_iNoise>
 	{
-		static void VM_FORCEINLINE DoKernel(float* RESTRICT Dst, float X, float Y, float Z)
+		static void VM_FORCEINLINE DoKernel(FVectorVMContext& Context, float* RESTRICT Dst, float X, float Y, float Z)
 		{
 			float Xfl = FMath::FloorToFloat(X);
 			float Yfl = FMath::FloorToFloat(Y);
@@ -785,7 +804,7 @@ namespace VectorVMNoise
 
 	struct FScalarKernelNoise2D_iNoise : TBinaryScalarKernel<FScalarKernelNoise2D_iNoise>
 	{
-		static void VM_FORCEINLINE DoKernel(float* RESTRICT Dst, float X, float Y)
+		static void VM_FORCEINLINE DoKernel(FVectorVMContext& Context, float* RESTRICT Dst, float X, float Y)
 		{
 			*Dst = 0.0f;//TODO
 		}
@@ -793,7 +812,7 @@ namespace VectorVMNoise
 
 	struct FScalarKernelNoise1D_iNoise : TUnaryScalarKernel<FScalarKernelNoise1D_iNoise>
 	{
-		static void VM_FORCEINLINE DoKernel(float* RESTRICT Dst, float X)
+		static void VM_FORCEINLINE DoKernel(FVectorVMContext& Context, float* RESTRICT Dst, float X)
 		{
 			*Dst = 0.0f;//TODO;
 		}
@@ -815,7 +834,7 @@ struct FVectorKernelNoise : public TUnaryVectorKernel<FVectorKernelNoise>
 {
 	static VectorRegister RandomTable[17][17][17];
 
-	static void VM_FORCEINLINE DoKernel(VectorRegister* RESTRICT Dst, VectorRegister Src0)
+	static void VM_FORCEINLINE DoKernel(FVectorVMContext& Context, VectorRegister* RESTRICT Dst, VectorRegister Src0)
 	{
 		const VectorRegister One = MakeVectorRegister(1.0f, 1.0f, 1.0f, 1.0f);
 		const VectorRegister VecSize = MakeVectorRegister(16.0f, 16.0f, 16.0f, 16.0f);
@@ -860,6 +879,94 @@ VectorRegister FVectorKernelNoise::RandomTable[17][17][17];
 //////////////////////////////////////////////////////////////////////////
 //Special Kernels.
 
+/** Special kernel for acquiring a new ID. TODO. Can be written as general RWBuffer ops when we support that. */
+struct FScalarKernelAcquireID
+{
+	static VM_FORCEINLINE void Exec(FVectorVMContext& Context)
+	{
+		int32 DataSetIndex = DecodeU16(Context);
+		TArray<int32>&RESTRICT FreeIDTable = *Context.DataSetMetaTable[DataSetIndex].FreeIDTable;
+
+		int32 Tag = Context.DataSetMetaTable[DataSetIndex].IDAcquireTag;
+
+		int32 IDIndexReg = DecodeU16(Context);
+		int32*RESTRICT IDIndex = (int32*)(Context.RegisterTable[IDIndexReg]);
+
+		int32 IDTagReg = DecodeU16(Context);
+		int32*RESTRICT IDTag = (int32*)(Context.RegisterTable[IDTagReg]);
+
+		int32* NumFreeIDs = Context.DataSetMetaTable[DataSetIndex].NumFreeIDs;
+	
+		check(FreeIDTable.Num() >= Context.NumInstances);
+		for (int32 i = 0; i < Context.NumInstances; ++i)
+		{
+			int32 FreeIDTableIndex = FPlatformAtomics::InterlockedDecrement(NumFreeIDs);
+			//UE_LOG(LogVectorVM, Warning, TEXT("AcquireID: InstanceID:%d | FreeTableIdx:%d | Tag:%d "), FreeIDTable[FreeIDTableIndex], FreeIDTableIndex, Tag);
+			*IDIndex = FreeIDTable[FreeIDTableIndex];
+			FreeIDTable[FreeIDTableIndex] = INDEX_NONE;
+			*IDTag = Tag;
+			++IDIndex;
+			++IDTag;
+		}
+	}
+};
+
+/** Special kernel for updating a new ID. TODO. Can be written as general RWBuffer ops when we support that. */
+struct FScalarKernelUpdateID
+{
+	static VM_FORCEINLINE void Exec(FVectorVMContext& Context)
+	{
+		int32 DataSetIndex = DecodeU16(Context);
+		int32 InstanceIDRegisterIndex = DecodeU16(Context);
+		int32 InstanceIndexRegisterIndex = DecodeU16(Context);
+		
+		TArray<int32>&RESTRICT IDTable = *Context.DataSetMetaTable[DataSetIndex].IDTable;
+		TArray<int32>&RESTRICT FreeIDTable = *Context.DataSetMetaTable[DataSetIndex].FreeIDTable;
+		int32 InstanceOffset = Context.DataSetMetaTable[DataSetIndex].InstanceOffset;// +Context.StartInstance;
+
+		int32*RESTRICT IDRegister = (int32*)(Context.RegisterTable[InstanceIDRegisterIndex]);
+		int32*RESTRICT IndexRegister = (int32*)(Context.RegisterTable[InstanceIndexRegisterIndex]);
+
+		volatile int32* MaxUsedID = Context.DataSetMetaTable[DataSetIndex].MaxUsedID;
+		int32* NumFreeIDs = Context.DataSetMetaTable[DataSetIndex].NumFreeIDs;
+
+		check(IDTable.Num() >= InstanceOffset + Context.NumInstances);
+		for (int32 i = 0; i < Context.NumInstances; ++i)
+		{
+			int32 InstanceId = IDRegister[i];
+			int32 Index = IndexRegister[i];
+
+			if (Index == INDEX_NONE)
+			{
+				//Instance not being written so release it's ID.
+				IDTable[InstanceId] = INDEX_NONE;
+				int32 FreeIDIndex = FPlatformAtomics::InterlockedIncrement(NumFreeIDs) - 1;
+				//UE_LOG(LogVectorVM, Warning, TEXT("UpdateID: Dead instance, Adding ID %d to free table at idx %d."), InstanceId, FreeIDIndex);
+				FreeIDTable[FreeIDIndex] = InstanceId;
+			}
+			else
+			{
+				int32 RealIdx = InstanceOffset + Index;
+				IDTable[InstanceId] = RealIdx;
+
+				//Update the MaxUsedID.
+				int32 LocalMaxUsedID;
+				do 
+				{
+					LocalMaxUsedID = *MaxUsedID;
+					if (LocalMaxUsedID >= InstanceId)
+					{
+						break;
+					}
+				} 
+				while (FPlatformAtomics::InterlockedCompareExchange(MaxUsedID, InstanceId, LocalMaxUsedID) != LocalMaxUsedID);
+				
+				//UE_LOG(LogVectorVM, Warning, TEXT("UpdateID: RealIdx:%d | InstanceID:%d."), RealIdx, InstanceId);
+			}
+		}
+	}
+};
+
 /** Special kernel for reading from the main input dataset. */
 template<typename T>
 struct FVectorKernelReadInput
@@ -956,7 +1063,9 @@ struct FOutputRegisterHandler : public FRegisterHandlerBase
 
 	VM_FORCEINLINE void Advance() { }
 	VM_FORCEINLINE T Get() { return *Register; }
-	VM_FORCEINLINE T* GetDest() { return Register; }
+	VM_FORCEINLINE T*RESTRICT GetDest() { return Register; }
+	VM_FORCEINLINE T*RESTRICT GetDestAndAdvance() { return Register; }
+	VM_FORCEINLINE T GetAndAdvance() { return *Register; }
 };
 
 /** Special kernel for writing to a specific output register. */
@@ -974,7 +1083,7 @@ struct FScalarKernelWriteOutputIndexed
 		};
 	}
 
-	static VM_FORCEINLINE void DoKernel(T* RESTRICT Dst, int32 SetOffset, int32 Index, T Data)
+	static VM_FORCEINLINE void DoKernel(FVectorVMContext& Context, T* RESTRICT Dst, int32 SetOffset, int32 Index, T Data)
 	{
 		if (Index != INDEX_NONE)
 		{
@@ -992,9 +1101,11 @@ struct FDataSetCounterHandler
 
 	VM_FORCEINLINE void Advance() { }
 	VM_FORCEINLINE int32* Get() { return Counter; }
+	VM_FORCEINLINE int32* GetAndAdvance() { return Counter; }
 	//VM_FORCEINLINE const int32* GetDest() { return Counter; }Should never use as a dest. All kernels with read and write to this.
 };
 
+template<bool bThreadsafe>
 struct FScalarKernelAcquireCounterIndex
 {
 	static VM_FORCEINLINE void Exec(FVectorVMContext& Context)
@@ -1008,11 +1119,18 @@ struct FScalarKernelAcquireCounterIndex
 		};
 	}
 
-	static VM_FORCEINLINE void DoKernel(int32* RESTRICT Dst, int32* Index, int32 Valid)
+	static VM_FORCEINLINE void DoKernel(FVectorVMContext& Context, int32* RESTRICT Dst, int32* Index, int32 Valid)
 	{
-		if (*Index != INDEX_NONE && Valid != 0)
+		if (Valid != 0)
 		{
-			*Dst = (*Index)++;
+			if (bThreadsafe)
+			{
+				*Dst = FPlatformAtomics::InterlockedIncrement(Index);
+			}
+			else
+			{
+				*Dst = ++(*Index);
+			}
 		}
 		else
 		{	
@@ -1058,7 +1176,7 @@ struct FKernelExternalFunctionCall
 //addi,
 struct FVectorIntKernelAdd : TBinaryVectorIntKernel<FVectorIntKernelAdd>
 {
-	static void VM_FORCEINLINE DoKernel(VectorRegisterInt* Dst, VectorRegisterInt Src0, VectorRegisterInt Src1)
+	static void VM_FORCEINLINE DoKernel(FVectorVMContext& Context, VectorRegisterInt* Dst, VectorRegisterInt Src0, VectorRegisterInt Src1)
 	{
 		*Dst = VectorIntAdd(Src0, Src1);
 	}
@@ -1067,7 +1185,7 @@ struct FVectorIntKernelAdd : TBinaryVectorIntKernel<FVectorIntKernelAdd>
 //subi,
 struct FVectorIntKernelSubtract : TBinaryVectorIntKernel<FVectorIntKernelSubtract>
 {
-	static void VM_FORCEINLINE DoKernel(VectorRegisterInt* Dst, VectorRegisterInt Src0, VectorRegisterInt Src1)
+	static void VM_FORCEINLINE DoKernel(FVectorVMContext& Context, VectorRegisterInt* Dst, VectorRegisterInt Src0, VectorRegisterInt Src1)
 	{
 		*Dst = VectorIntSubtract(Src0, Src1);
 	}
@@ -1076,7 +1194,7 @@ struct FVectorIntKernelSubtract : TBinaryVectorIntKernel<FVectorIntKernelSubtrac
 //muli,
 struct FVectorIntKernelMultiply : TBinaryVectorIntKernel<FVectorIntKernelMultiply>
 {
-	static void VM_FORCEINLINE DoKernel(VectorRegisterInt* Dst, VectorRegisterInt Src0, VectorRegisterInt Src1)
+	static void VM_FORCEINLINE DoKernel(FVectorVMContext& Context, VectorRegisterInt* Dst, VectorRegisterInt Src0, VectorRegisterInt Src1)
 	{
 		*Dst = VectorIntMultiply(Src0, Src1);
 	}
@@ -1085,7 +1203,7 @@ struct FVectorIntKernelMultiply : TBinaryVectorIntKernel<FVectorIntKernelMultipl
 //clampi,
 struct FVectorIntKernelClamp : TTrinaryVectorIntKernel<FVectorIntKernelClamp>
 {
-	static void VM_FORCEINLINE DoKernel(VectorRegisterInt* Dst, VectorRegisterInt Src0, VectorRegisterInt Src1, VectorRegisterInt Src2)
+	static void VM_FORCEINLINE DoKernel(FVectorVMContext& Context, VectorRegisterInt* Dst, VectorRegisterInt Src0, VectorRegisterInt Src1, VectorRegisterInt Src2)
 	{
 		*Dst = VectorIntMin(VectorIntMax(Src0, Src1), Src2);
 	}
@@ -1094,7 +1212,7 @@ struct FVectorIntKernelClamp : TTrinaryVectorIntKernel<FVectorIntKernelClamp>
 //mini,
 struct FVectorIntKernelMin : TBinaryVectorIntKernel<FVectorIntKernelMin>
 {
-	static void VM_FORCEINLINE DoKernel(VectorRegisterInt* Dst, VectorRegisterInt Src0, VectorRegisterInt Src1)
+	static void VM_FORCEINLINE DoKernel(FVectorVMContext& Context, VectorRegisterInt* Dst, VectorRegisterInt Src0, VectorRegisterInt Src1)
 	{
 		*Dst = VectorIntMin(Src0, Src1);
 	}
@@ -1103,7 +1221,7 @@ struct FVectorIntKernelMin : TBinaryVectorIntKernel<FVectorIntKernelMin>
 //maxi,
 struct FVectorIntKernelMax : TBinaryVectorIntKernel<FVectorIntKernelMax>
 {
-	static void VM_FORCEINLINE DoKernel(VectorRegisterInt* Dst, VectorRegisterInt Src0, VectorRegisterInt Src1)
+	static void VM_FORCEINLINE DoKernel(FVectorVMContext& Context, VectorRegisterInt* Dst, VectorRegisterInt Src0, VectorRegisterInt Src1)
 	{
 		*Dst = VectorIntMax(Src0, Src1);
 	}
@@ -1112,7 +1230,7 @@ struct FVectorIntKernelMax : TBinaryVectorIntKernel<FVectorIntKernelMax>
 //absi,
 struct FVectorIntKernelAbs : TUnaryVectorIntKernel<FVectorIntKernelAbs>
 {
-	static void VM_FORCEINLINE DoKernel(VectorRegisterInt* Dst, VectorRegisterInt Src0)
+	static void VM_FORCEINLINE DoKernel(FVectorVMContext& Context, VectorRegisterInt* Dst, VectorRegisterInt Src0)
 	{
 		*Dst = VectorIntAbs(Src0);
 	}
@@ -1121,7 +1239,7 @@ struct FVectorIntKernelAbs : TUnaryVectorIntKernel<FVectorIntKernelAbs>
 //negi,
 struct FVectorIntKernelNegate : TUnaryVectorIntKernel<FVectorIntKernelNegate>
 {
-	static void VM_FORCEINLINE DoKernel(VectorRegisterInt* Dst, VectorRegisterInt Src0)
+	static void VM_FORCEINLINE DoKernel(FVectorVMContext& Context, VectorRegisterInt* Dst, VectorRegisterInt Src0)
 	{
 		*Dst = VectorIntNegate(Src0);
 	}
@@ -1130,7 +1248,7 @@ struct FVectorIntKernelNegate : TUnaryVectorIntKernel<FVectorIntKernelNegate>
 //signi,
 struct FVectorIntKernelSign : TUnaryVectorIntKernel<FVectorIntKernelSign>
 {
-	static void VM_FORCEINLINE DoKernel(VectorRegisterInt* Dst, VectorRegisterInt Src0)
+	static void VM_FORCEINLINE DoKernel(FVectorVMContext& Context, VectorRegisterInt* Dst, VectorRegisterInt Src0)
 	{
 		*Dst = VectorIntSign(Src0);
 	}
@@ -1140,18 +1258,18 @@ struct FVectorIntKernelSign : TUnaryVectorIntKernel<FVectorIntKernelSign>
 //No good way to do this with SSE atm so just do it scalar.
 struct FScalarIntKernelRandom : public TUnaryScalarIntKernel<FScalarIntKernelRandom>
 {
-	static void VM_FORCEINLINE DoKernel(int32* RESTRICT Dst, int32 Src0)
+	static void VM_FORCEINLINE DoKernel(FVectorVMContext& Context, int32* RESTRICT Dst, int32 Src0)
 	{
 		const float rm = RAND_MAX;
 		//EEK!. Improve this. Implement GPU style seeded rand instead of this.
-		*Dst = static_cast<int32>((FMath::Rand() / rm) * Src0);
+		*Dst = static_cast<int32>(Context.RandStream.GetFraction() * Src0);
 	}
 };
 
 //cmplti,
 struct FVectorIntKernelCompareLT : TBinaryVectorIntKernel<FVectorIntKernelCompareLT>
 {
-	static void VM_FORCEINLINE DoKernel(VectorRegisterInt* Dst, VectorRegisterInt Src0, VectorRegisterInt Src1)
+	static void VM_FORCEINLINE DoKernel(FVectorVMContext& Context, VectorRegisterInt* Dst, VectorRegisterInt Src0, VectorRegisterInt Src1)
 	{
 		*Dst = VectorIntCompareLT(Src0, Src1);
 	}
@@ -1160,7 +1278,7 @@ struct FVectorIntKernelCompareLT : TBinaryVectorIntKernel<FVectorIntKernelCompar
 //cmplei,
 struct FVectorIntKernelCompareLE : TBinaryVectorIntKernel<FVectorIntKernelCompareLE>
 {
-	static void VM_FORCEINLINE DoKernel(VectorRegisterInt* Dst, VectorRegisterInt Src0, VectorRegisterInt Src1)
+	static void VM_FORCEINLINE DoKernel(FVectorVMContext& Context, VectorRegisterInt* Dst, VectorRegisterInt Src0, VectorRegisterInt Src1)
 	{
 		*Dst = VectorIntCompareLE(Src0, Src1);
 	}
@@ -1169,7 +1287,7 @@ struct FVectorIntKernelCompareLE : TBinaryVectorIntKernel<FVectorIntKernelCompar
 //cmpgti,
 struct FVectorIntKernelCompareGT : TBinaryVectorIntKernel<FVectorIntKernelCompareGT>
 {
-	static void VM_FORCEINLINE DoKernel(VectorRegisterInt* Dst, VectorRegisterInt Src0, VectorRegisterInt Src1)
+	static void VM_FORCEINLINE DoKernel(FVectorVMContext& Context, VectorRegisterInt* Dst, VectorRegisterInt Src0, VectorRegisterInt Src1)
 	{
 		*Dst = VectorIntCompareGT(Src0, Src1);
 	}
@@ -1178,7 +1296,7 @@ struct FVectorIntKernelCompareGT : TBinaryVectorIntKernel<FVectorIntKernelCompar
 //cmpgei,
 struct FVectorIntKernelCompareGE : TBinaryVectorIntKernel<FVectorIntKernelCompareGE>
 {
-	static void VM_FORCEINLINE DoKernel(VectorRegisterInt* Dst, VectorRegisterInt Src0, VectorRegisterInt Src1)
+	static void VM_FORCEINLINE DoKernel(FVectorVMContext& Context, VectorRegisterInt* Dst, VectorRegisterInt Src0, VectorRegisterInt Src1)
 	{
 		*Dst = VectorIntCompareGE(Src0, Src1);
 	}
@@ -1187,7 +1305,7 @@ struct FVectorIntKernelCompareGE : TBinaryVectorIntKernel<FVectorIntKernelCompar
 //cmpeqi,
 struct FVectorIntKernelCompareEQ : TBinaryVectorIntKernel<FVectorIntKernelCompareEQ>
 {
-	static void VM_FORCEINLINE DoKernel(VectorRegisterInt* Dst, VectorRegisterInt Src0, VectorRegisterInt Src1)
+	static void VM_FORCEINLINE DoKernel(FVectorVMContext& Context, VectorRegisterInt* Dst, VectorRegisterInt Src0, VectorRegisterInt Src1)
 	{
 		*Dst = VectorIntCompareEQ(Src0, Src1);
 	}
@@ -1196,7 +1314,7 @@ struct FVectorIntKernelCompareEQ : TBinaryVectorIntKernel<FVectorIntKernelCompar
 //cmpneqi,
 struct FVectorIntKernelCompareNEQ : TBinaryVectorIntKernel<FVectorIntKernelCompareNEQ>
 {
-	static void VM_FORCEINLINE DoKernel(VectorRegisterInt* Dst, VectorRegisterInt Src0, VectorRegisterInt Src1)
+	static void VM_FORCEINLINE DoKernel(FVectorVMContext& Context, VectorRegisterInt* Dst, VectorRegisterInt Src0, VectorRegisterInt Src1)
 	{
 		*Dst = VectorIntCompareNEQ(Src0, Src1);
 	}
@@ -1205,7 +1323,7 @@ struct FVectorIntKernelCompareNEQ : TBinaryVectorIntKernel<FVectorIntKernelCompa
 //bit_and,
 struct FVectorIntKernelBitAnd : TBinaryVectorIntKernel<FVectorIntKernelBitAnd>
 {
-	static void VM_FORCEINLINE DoKernel(VectorRegisterInt* Dst, VectorRegisterInt Src0, VectorRegisterInt Src1)
+	static void VM_FORCEINLINE DoKernel(FVectorVMContext& Context, VectorRegisterInt* Dst, VectorRegisterInt Src0, VectorRegisterInt Src1)
 	{
 		*Dst = VectorIntAnd(Src0, Src1);
 	}
@@ -1214,7 +1332,7 @@ struct FVectorIntKernelBitAnd : TBinaryVectorIntKernel<FVectorIntKernelBitAnd>
 //bit_or,
 struct FVectorIntKernelBitOr : TBinaryVectorIntKernel<FVectorIntKernelBitOr>
 {
-	static void VM_FORCEINLINE DoKernel(VectorRegisterInt* Dst, VectorRegisterInt Src0, VectorRegisterInt Src1)
+	static void VM_FORCEINLINE DoKernel(FVectorVMContext& Context, VectorRegisterInt* Dst, VectorRegisterInt Src0, VectorRegisterInt Src1)
 	{
 		*Dst = VectorIntOr(Src0, Src1);
 	}
@@ -1223,7 +1341,7 @@ struct FVectorIntKernelBitOr : TBinaryVectorIntKernel<FVectorIntKernelBitOr>
 //bit_xor,
 struct FVectorIntKernelBitXor : TBinaryVectorIntKernel<FVectorIntKernelBitXor>
 {
-	static void VM_FORCEINLINE DoKernel(VectorRegisterInt* Dst, VectorRegisterInt Src0, VectorRegisterInt Src1)
+	static void VM_FORCEINLINE DoKernel(FVectorVMContext& Context, VectorRegisterInt* Dst, VectorRegisterInt Src0, VectorRegisterInt Src1)
 	{
 		*Dst = VectorIntXor(Src0, Src1);
 	}
@@ -1232,7 +1350,7 @@ struct FVectorIntKernelBitXor : TBinaryVectorIntKernel<FVectorIntKernelBitXor>
 //bit_not,
 struct FVectorIntKernelBitNot : TUnaryVectorIntKernel<FVectorIntKernelBitNot>
 {
-	static void VM_FORCEINLINE DoKernel(VectorRegisterInt* Dst, VectorRegisterInt Src0)
+	static void VM_FORCEINLINE DoKernel(FVectorVMContext& Context, VectorRegisterInt* Dst, VectorRegisterInt Src0)
 	{
 		*Dst = VectorIntNot(Src0);
 	}
@@ -1242,7 +1360,7 @@ struct FVectorIntKernelBitNot : TUnaryVectorIntKernel<FVectorIntKernelBitNot>
 //logic_and,
 struct FVectorIntKernelLogicAnd : TBinaryVectorIntKernel<FVectorIntKernelLogicAnd>
 {
-	static void VM_FORCEINLINE DoKernel(VectorRegisterInt* Dst, VectorRegisterInt Src0, VectorRegisterInt Src1)
+	static void VM_FORCEINLINE DoKernel(FVectorVMContext& Context, VectorRegisterInt* Dst, VectorRegisterInt Src0, VectorRegisterInt Src1)
 	{
 		//We need to assume a mask input and produce a mask output so just bitwise ops actually fine for these?
 		*Dst = VectorIntAnd(Src0, Src1);
@@ -1252,7 +1370,7 @@ struct FVectorIntKernelLogicAnd : TBinaryVectorIntKernel<FVectorIntKernelLogicAn
 //logic_or,
 struct FVectorIntKernelLogicOr : TBinaryVectorIntKernel<FVectorIntKernelLogicOr>
 {
-	static void VM_FORCEINLINE DoKernel(VectorRegisterInt* Dst, VectorRegisterInt Src0, VectorRegisterInt Src1)
+	static void VM_FORCEINLINE DoKernel(FVectorVMContext& Context, VectorRegisterInt* Dst, VectorRegisterInt Src0, VectorRegisterInt Src1)
 	{
 		//We need to assume a mask input and produce a mask output so just bitwise ops actually fine for these?
 		*Dst = VectorIntOr(Src0, Src1);
@@ -1261,7 +1379,7 @@ struct FVectorIntKernelLogicOr : TBinaryVectorIntKernel<FVectorIntKernelLogicOr>
 //logic_xor,
 struct FVectorIntKernelLogicXor : TBinaryVectorIntKernel<FVectorIntKernelLogicXor>
 {
-	static void VM_FORCEINLINE DoKernel(VectorRegisterInt* Dst, VectorRegisterInt Src0, VectorRegisterInt Src1)
+	static void VM_FORCEINLINE DoKernel(FVectorVMContext& Context, VectorRegisterInt* Dst, VectorRegisterInt Src0, VectorRegisterInt Src1)
 	{
 		//We need to assume a mask input and produce a mask output so just bitwise ops actually fine for these?
 		*Dst = VectorIntXor(Src0, Src1);
@@ -1271,7 +1389,7 @@ struct FVectorIntKernelLogicXor : TBinaryVectorIntKernel<FVectorIntKernelLogicXo
 //logic_not,
 struct FVectorIntKernelLogicNot : TUnaryVectorIntKernel<FVectorIntKernelLogicNot>
 {
-	static void VM_FORCEINLINE DoKernel(VectorRegisterInt* Dst, VectorRegisterInt Src0)
+	static void VM_FORCEINLINE DoKernel(FVectorVMContext& Context, VectorRegisterInt* Dst, VectorRegisterInt Src0)
 	{
 		//We need to assume a mask input and produce a mask output so just bitwise ops actually fine for these?
 		*Dst = VectorIntNot(Src0);
@@ -1282,7 +1400,7 @@ struct FVectorIntKernelLogicNot : TUnaryVectorIntKernel<FVectorIntKernelLogicNot
 //f2i,
 struct FVectorKernelFloatToInt : TUnaryKernel<FVectorKernelFloatToInt, FRegisterDestHandler<VectorRegisterInt>, FConstantHandler<VectorRegister>, FRegisterHandler<VectorRegister>, VECTOR_WIDTH_FLOATS>
 {
-	static void VM_FORCEINLINE DoKernel(VectorRegisterInt* Dst, VectorRegister Src0)
+	static void VM_FORCEINLINE DoKernel(FVectorVMContext& Context, VectorRegisterInt* Dst, VectorRegister Src0)
 	{
 		*Dst = VectorFloatToInt(Src0);
 	}
@@ -1291,7 +1409,7 @@ struct FVectorKernelFloatToInt : TUnaryKernel<FVectorKernelFloatToInt, FRegister
 //i2f,
 struct FVectorKernelIntToFloat : TUnaryKernel<FVectorKernelIntToFloat, FRegisterDestHandler<VectorRegister>, FConstantHandler<VectorRegisterInt>, FRegisterHandler<VectorRegisterInt>, VECTOR_WIDTH_FLOATS>
 {
-	static void VM_FORCEINLINE DoKernel(VectorRegister* Dst, VectorRegisterInt Src0)
+	static void VM_FORCEINLINE DoKernel(FVectorVMContext& Context, VectorRegister* Dst, VectorRegisterInt Src0)
 	{
 		*Dst = VectorIntToFloat(Src0);
 	}
@@ -1300,7 +1418,7 @@ struct FVectorKernelIntToFloat : TUnaryKernel<FVectorKernelIntToFloat, FRegister
 //f2b,
 struct FVectorKernelFloatToBool : TUnaryKernel<FVectorKernelFloatToBool, FRegisterDestHandler<VectorRegister>, FConstantHandler<VectorRegister>, FRegisterHandler<VectorRegister>, VECTOR_WIDTH_FLOATS>
 {
-	static void VM_FORCEINLINE DoKernel(VectorRegister* Dst, VectorRegister Src0)
+	static void VM_FORCEINLINE DoKernel(FVectorVMContext& Context, VectorRegister* Dst, VectorRegister Src0)
 	{		
 		*Dst = VectorCompareGT(Src0, GlobalVectorConstants::FloatZero);
 	}
@@ -1309,7 +1427,7 @@ struct FVectorKernelFloatToBool : TUnaryKernel<FVectorKernelFloatToBool, FRegist
 //b2f,
 struct FVectorKernelBoolToFloat : TUnaryKernel<FVectorKernelBoolToFloat, FRegisterDestHandler<VectorRegister>, FConstantHandler<VectorRegister>, FRegisterHandler<VectorRegister>, VECTOR_WIDTH_FLOATS>
 {
-	static void VM_FORCEINLINE DoKernel(VectorRegister* Dst, VectorRegister Src0)
+	static void VM_FORCEINLINE DoKernel(FVectorVMContext& Context, VectorRegister* Dst, VectorRegister Src0)
 	{
 		*Dst = VectorSelect(Src0, GlobalVectorConstants::FloatOne, GlobalVectorConstants::FloatZero);
 	}
@@ -1318,7 +1436,7 @@ struct FVectorKernelBoolToFloat : TUnaryKernel<FVectorKernelBoolToFloat, FRegist
 //i2b,
 struct FVectorKernelIntToBool : TUnaryKernel<FVectorKernelIntToBool, FRegisterDestHandler<VectorRegisterInt>, FConstantHandler<VectorRegisterInt>, FRegisterHandler<VectorRegisterInt>, VECTOR_WIDTH_FLOATS>
 {
-	static void VM_FORCEINLINE DoKernel(VectorRegisterInt* Dst, VectorRegisterInt Src0)
+	static void VM_FORCEINLINE DoKernel(FVectorVMContext& Context, VectorRegisterInt* Dst, VectorRegisterInt Src0)
 	{
 		*Dst = VectorIntCompareGT(Src0, GlobalVectorConstants::IntZero);
 	}
@@ -1327,18 +1445,24 @@ struct FVectorKernelIntToBool : TUnaryKernel<FVectorKernelIntToBool, FRegisterDe
 //b2i,
 struct FVectorKernelBoolToInt : TUnaryKernel<FVectorKernelBoolToInt, FRegisterDestHandler<VectorRegisterInt>, FConstantHandler<VectorRegisterInt>, FRegisterHandler<VectorRegisterInt>, VECTOR_WIDTH_FLOATS>
 {
-	static void VM_FORCEINLINE DoKernel(VectorRegisterInt* Dst, VectorRegisterInt Src0)
+	static void VM_FORCEINLINE DoKernel(FVectorVMContext& Context, VectorRegisterInt* Dst, VectorRegisterInt Src0)
 	{
 		*Dst = VectorIntSelect(Src0, GlobalVectorConstants::IntOne, GlobalVectorConstants::IntZero);
 	}
 };
 
+UEnum* g_VectorVMEnumStateObj = nullptr;
+UEnum* g_VectorVMEnumOperandObj = nullptr;
+
 
 void VectorVM::Init()
 {
 	static bool Inited = false;
 	if (Inited == false)
 	{
+		g_VectorVMEnumStateObj = FindObject<UEnum>(ANY_PACKAGE, TEXT("EVectorVMOp"), true);
+		g_VectorVMEnumOperandObj = FindObject<UEnum>(ANY_PACKAGE, TEXT("EVectorVMOperandLocation"), true);
+
 		// random noise
 		float TempTable[17][17][17];
 		for (int z = 0; z < 17; z++)
@@ -1429,7 +1553,6 @@ void VectorVM::Exec(
 	FVMExternalFunction* ExternalFunctionTable,
 	void** UserPtrTable,
 	int32 NumInstances
-
 #if STATS
 	, const TArray<TStatId>& StatScopes
 #endif
@@ -1450,142 +1573,170 @@ void VectorVM::Exec(
 		DataSetIndexTable.Add(DataSetMetaTable[Idx].DataSetAccessIndex);	// prime counter index table with the data set offset; will be incremented with every write for each instance
 	}
 
-	FVectorVMContext& Context = FVectorVMContext::Get();
-	Context.PrepareForExec(InputRegisters, OutputRegisters, NumInputRegisters, NumOutputRegisters, ConstantTable, DataSetIndexTable.GetData(), DataSetOffsetTable.GetData(), DataSetOffsetTable.Num(),
-		ExternalFunctionTable, UserPtrTable
+	int32 NumChunks = (NumInstances / InstancesPerChunk) + 1;
+	int32 ChunksPerBatch = GbParallelVVM != 0 ? GParallelVVMChunksPerBatch : NumChunks;
+	int32 NumBatches = FMath::DivideAndRoundUp(NumChunks, ChunksPerBatch);
+	bool bParallel = NumBatches > 1;
+
+	auto ExecChunkBatch = [&](int32 BatchIdx)
+	{		
+		FVectorVMContext& Context = FVectorVMContext::Get();
+		Context.PrepareForExec(InputRegisters, OutputRegisters, NumInputRegisters, NumOutputRegisters, ConstantTable, DataSetIndexTable.GetData(), DataSetOffsetTable.GetData(), DataSetOffsetTable.Num(),
+			ExternalFunctionTable, UserPtrTable, DataSetMetaTable.GetData()
 #if STATS
-		, &StatScopes
+			, &StatScopes
 #endif
-	);
-
-	// Process one chunk at a time.
-	int32 InstancesLeft = NumInstances;
-	int32 ChunkIdx = 0;
-	while (InstancesLeft > 0)
-	{
-		// Setup execution context.
-		Context.PrepareForChunk(Code, FMath::Min(InstancesLeft, (int32)InstancesPerChunk), InstancesPerChunk * ChunkIdx);
+		);
+
+		// Process one chunk at a time.
+		int32 ChunkIdx = BatchIdx * ChunksPerBatch;
+		int32 FirstInstance = ChunkIdx * InstancesPerChunk;
+		int32 FinalInstance = FMath::Min(NumInstances, FirstInstance + (ChunksPerBatch * InstancesPerChunk));
+		int32 InstancesLeft = FinalInstance - FirstInstance;
+		while (InstancesLeft > 0)
+		{
+			int32 NumInstancesThisChunk = FMath::Min(InstancesLeft, (int32)InstancesPerChunk);
+			int32 StartInstance = InstancesPerChunk * ChunkIdx;
+			// Setup execution context.
+			Context.PrepareForChunk(Code, NumInstancesThisChunk, StartInstance);
 
-		EVectorVMOp Op = EVectorVMOp::done;
+			EVectorVMOp Op = EVectorVMOp::done;
 
-		// Execute VM on all vectors in this chunk.
-		do 
-		{
-			Op = DecodeOp(Context);
-			switch (Op)
+			// Execute VM on all vectors in this chunk.
+			do
 			{
-			// Dispatch kernel ops.
-			case EVectorVMOp::add: FVectorKernelAdd::Exec(Context); break;
-			case EVectorVMOp::sub: FVectorKernelSub::Exec(Context); break;
-			case EVectorVMOp::mul: FVectorKernelMul::Exec(Context); break;
-			case EVectorVMOp::div: FVectorKernelDiv::Exec(Context); break;
-			case EVectorVMOp::mad: FVectorKernelMad::Exec(Context); break;
-			case EVectorVMOp::lerp: FVectorKernelLerp::Exec(Context); break;
-			case EVectorVMOp::rcp: FVectorKernelRcp::Exec(Context); break;
-			case EVectorVMOp::rsq: FVectorKernelRsq::Exec(Context); break;
-			case EVectorVMOp::sqrt: FVectorKernelSqrt::Exec(Context); break;
-			case EVectorVMOp::neg: FVectorKernelNeg::Exec(Context); break;
-			case EVectorVMOp::abs: FVectorKernelAbs::Exec(Context); break;
-			case EVectorVMOp::exp: FVectorKernelExp::Exec(Context); break;
-			case EVectorVMOp::exp2: FVectorKernelExp2::Exec(Context); break;
-			case EVectorVMOp::log: FVectorKernelLog::Exec(Context); break;
-			case EVectorVMOp::log2: FVectorKernelLog2::Exec(Context); break;
-			case EVectorVMOp::sin: FVectorKernelSin::Exec(Context); break;
-			case EVectorVMOp::cos: FVectorKernelCos::Exec(Context); break;
-			case EVectorVMOp::tan: FVectorKernelTan::Exec(Context); break;
-			case EVectorVMOp::asin: FVectorKernelASin::Exec(Context); break;
-			case EVectorVMOp::acos: FVectorKernelACos::Exec(Context); break;
-			case EVectorVMOp::atan: FVectorKernelATan::Exec(Context); break;
-			case EVectorVMOp::atan2: FVectorKernelATan2::Exec(Context); break;
-			case EVectorVMOp::ceil: FVectorKernelCeil::Exec(Context); break;
-			case EVectorVMOp::floor: FVectorKernelFloor::Exec(Context); break;
-			case EVectorVMOp::round: FVectorKernelRound::Exec(Context); break;
-			case EVectorVMOp::fmod: FVectorKernelMod::Exec(Context); break;
-			case EVectorVMOp::frac: FVectorKernelFrac::Exec(Context); break;
-			case EVectorVMOp::trunc: FVectorKernelTrunc::Exec(Context); break;
-			case EVectorVMOp::clamp: FVectorKernelClamp::Exec(Context); break;
-			case EVectorVMOp::min: FVectorKernelMin::Exec(Context); break;
-			case EVectorVMOp::max: FVectorKernelMax::Exec(Context); break;
-			case EVectorVMOp::pow: FVectorKernelPow::Exec(Context); break;
-			case EVectorVMOp::sign: FVectorKernelSign::Exec(Context); break;
-			case EVectorVMOp::step: FVectorKernelStep::Exec(Context); break;
-			case EVectorVMOp::random: FVectorKernelRandom::Exec(Context); break;
-			case EVectorVMOp::noise: VectorVMNoise::Noise1D(Context); break;
-			case EVectorVMOp::noise2D: VectorVMNoise::Noise2D(Context); break;
-			case EVectorVMOp::noise3D: VectorVMNoise::Noise3D(Context); break;
-
-			case EVectorVMOp::cmplt: FVectorKernelCompareLT::Exec(Context); break;
-			case EVectorVMOp::cmple: FVectorKernelCompareLE::Exec(Context); break;
-			case EVectorVMOp::cmpgt: FVectorKernelCompareGT::Exec(Context); break;
-			case EVectorVMOp::cmpge: FVectorKernelCompareGE::Exec(Context); break;
-			case EVectorVMOp::cmpeq: FVectorKernelCompareEQ::Exec(Context); break;
-			case EVectorVMOp::cmpneq: FVectorKernelCompareNEQ::Exec(Context); break;
-			case EVectorVMOp::select: FVectorKernelSelect::Exec(Context); break;
-
-			case EVectorVMOp::addi: FVectorIntKernelAdd::Exec(Context); break;
-			case EVectorVMOp::subi: FVectorIntKernelSubtract::Exec(Context); break;
-			case EVectorVMOp::muli: FVectorIntKernelMultiply::Exec(Context); break;
-			case EVectorVMOp::clampi: FVectorIntKernelClamp::Exec(Context); break;
-			case EVectorVMOp::mini: FVectorIntKernelMin::Exec(Context); break;
-			case EVectorVMOp::maxi: FVectorIntKernelMax::Exec(Context); break;
-			case EVectorVMOp::absi: FVectorIntKernelAbs::Exec(Context); break;
-			case EVectorVMOp::negi: FVectorIntKernelNegate::Exec(Context); break;
-			case EVectorVMOp::signi: FVectorIntKernelSign::Exec(Context); break;
-			case EVectorVMOp::randomi: FScalarIntKernelRandom::Exec(Context); break;
-			case EVectorVMOp::cmplti: FVectorIntKernelCompareLT::Exec(Context); break;
-			case EVectorVMOp::cmplei: FVectorIntKernelCompareLE::Exec(Context); break;
-			case EVectorVMOp::cmpgti: FVectorIntKernelCompareGT::Exec(Context); break;
-			case EVectorVMOp::cmpgei: FVectorIntKernelCompareGE::Exec(Context); break;
-			case EVectorVMOp::cmpeqi: FVectorIntKernelCompareEQ::Exec(Context); break;
-			case EVectorVMOp::cmpneqi: FVectorIntKernelCompareNEQ::Exec(Context); break;
-			case EVectorVMOp::bit_and: FVectorIntKernelBitAnd::Exec(Context); break;
-			case EVectorVMOp::bit_or: FVectorIntKernelBitOr::Exec(Context); break;
-			case EVectorVMOp::bit_xor: FVectorIntKernelBitXor::Exec(Context); break;
-			case EVectorVMOp::bit_not: FVectorIntKernelBitNot::Exec(Context); break;
-			case EVectorVMOp::logic_and: FVectorIntKernelLogicAnd::Exec(Context); break;
-			case EVectorVMOp::logic_or: FVectorIntKernelLogicOr::Exec(Context); break;
-			case EVectorVMOp::logic_xor: FVectorIntKernelLogicXor::Exec(Context); break;
-			case EVectorVMOp::logic_not: FVectorIntKernelLogicNot::Exec(Context); break;
-			case EVectorVMOp::f2i: FVectorKernelFloatToInt::Exec(Context); break;
-			case EVectorVMOp::i2f: FVectorKernelIntToFloat::Exec(Context); break;
-			case EVectorVMOp::f2b: FVectorKernelFloatToBool::Exec(Context); break;
-			case EVectorVMOp::b2f: FVectorKernelBoolToFloat::Exec(Context); break;
-			case EVectorVMOp::i2b: FVectorKernelIntToBool::Exec(Context); break;
-			case EVectorVMOp::b2i: FVectorKernelBoolToInt::Exec(Context); break;
-
-			case EVectorVMOp::outputdata_32bit:	FScalarKernelWriteOutputIndexed<int32>::Exec(Context);	break;
-			case EVectorVMOp::inputdata_32bit: FVectorKernelReadInput<int32>::Exec(Context); break;
-			//case EVectorVMOp::inputdata_32bit: FVectorKernelReadInput32::Exec(Context); break;
-			case EVectorVMOp::inputdata_noadvance_32bit: FVectorKernelReadInputNoAdvance<int32>::Exec(Context); break;
-			case EVectorVMOp::acquireindex:	FScalarKernelAcquireCounterIndex::Exec(Context); break;
-			case EVectorVMOp::external_func_call: FKernelExternalFunctionCall::Exec(Context); break;
-
-			case EVectorVMOp::exec_index: FVectorKernelExecutionIndex::Exec(Context); break;
-
-			case EVectorVMOp::enter_stat_scope: FVectorKernelEnterStatScope::Exec(Context); break;
-			case EVectorVMOp::exit_stat_scope: FVectorKernelExitStatScope::Exec(Context); break;
-
-			// Execution always terminates with a "done" opcode.
-			case EVectorVMOp::done:
-				break;
-
-			// Opcode not recognized / implemented.
-			default:
-				UE_LOG(LogVectorVM, Fatal, TEXT("Unknown op code 0x%02x"), (uint32)Op);
-				return;//BAIL
-			}
-		} while (Op != EVectorVMOp::done);
+				Op = DecodeOp(Context);
+				switch (Op)
+				{
+					// Dispatch kernel ops.
+				case EVectorVMOp::add: FVectorKernelAdd::Exec(Context); break;
+				case EVectorVMOp::sub: FVectorKernelSub::Exec(Context); break;
+				case EVectorVMOp::mul: FVectorKernelMul::Exec(Context); break;
+				case EVectorVMOp::div: FVectorKernelDiv::Exec(Context); break;
+				case EVectorVMOp::mad: FVectorKernelMad::Exec(Context); break;
+				case EVectorVMOp::lerp: FVectorKernelLerp::Exec(Context); break;
+				case EVectorVMOp::rcp: FVectorKernelRcp::Exec(Context); break;
+				case EVectorVMOp::rsq: FVectorKernelRsq::Exec(Context); break;
+				case EVectorVMOp::sqrt: FVectorKernelSqrt::Exec(Context); break;
+				case EVectorVMOp::neg: FVectorKernelNeg::Exec(Context); break;
+				case EVectorVMOp::abs: FVectorKernelAbs::Exec(Context); break;
+				case EVectorVMOp::exp: FVectorKernelExp::Exec(Context); break;
+				case EVectorVMOp::exp2: FVectorKernelExp2::Exec(Context); break;
+				case EVectorVMOp::log: FVectorKernelLog::Exec(Context); break;
+				case EVectorVMOp::log2: FVectorKernelLog2::Exec(Context); break;
+				case EVectorVMOp::sin: FVectorKernelSin::Exec(Context); break;
+				case EVectorVMOp::cos: FVectorKernelCos::Exec(Context); break;
+				case EVectorVMOp::tan: FVectorKernelTan::Exec(Context); break;
+				case EVectorVMOp::asin: FVectorKernelASin::Exec(Context); break;
+				case EVectorVMOp::acos: FVectorKernelACos::Exec(Context); break;
+				case EVectorVMOp::atan: FVectorKernelATan::Exec(Context); break;
+				case EVectorVMOp::atan2: FVectorKernelATan2::Exec(Context); break;
+				case EVectorVMOp::ceil: FVectorKernelCeil::Exec(Context); break;
+				case EVectorVMOp::floor: FVectorKernelFloor::Exec(Context); break;
+				case EVectorVMOp::round: FVectorKernelRound::Exec(Context); break;
+				case EVectorVMOp::fmod: FVectorKernelMod::Exec(Context); break;
+				case EVectorVMOp::frac: FVectorKernelFrac::Exec(Context); break;
+				case EVectorVMOp::trunc: FVectorKernelTrunc::Exec(Context); break;
+				case EVectorVMOp::clamp: FVectorKernelClamp::Exec(Context); break;
+				case EVectorVMOp::min: FVectorKernelMin::Exec(Context); break;
+				case EVectorVMOp::max: FVectorKernelMax::Exec(Context); break;
+				case EVectorVMOp::pow: FVectorKernelPow::Exec(Context); break;
+				case EVectorVMOp::sign: FVectorKernelSign::Exec(Context); break;
+				case EVectorVMOp::step: FVectorKernelStep::Exec(Context); break;
+				case EVectorVMOp::random: FVectorKernelRandom::Exec(Context); break;
+				case EVectorVMOp::noise: VectorVMNoise::Noise1D(Context); break;
+				case EVectorVMOp::noise2D: VectorVMNoise::Noise2D(Context); break;
+				case EVectorVMOp::noise3D: VectorVMNoise::Noise3D(Context); break;
+
+				case EVectorVMOp::cmplt: FVectorKernelCompareLT::Exec(Context); break;
+				case EVectorVMOp::cmple: FVectorKernelCompareLE::Exec(Context); break;
+				case EVectorVMOp::cmpgt: FVectorKernelCompareGT::Exec(Context); break;
+				case EVectorVMOp::cmpge: FVectorKernelCompareGE::Exec(Context); break;
+				case EVectorVMOp::cmpeq: FVectorKernelCompareEQ::Exec(Context); break;
+				case EVectorVMOp::cmpneq: FVectorKernelCompareNEQ::Exec(Context); break;
+				case EVectorVMOp::select: FVectorKernelSelect::Exec(Context); break;
+
+				case EVectorVMOp::addi: FVectorIntKernelAdd::Exec(Context); break;
+				case EVectorVMOp::subi: FVectorIntKernelSubtract::Exec(Context); break;
+				case EVectorVMOp::muli: FVectorIntKernelMultiply::Exec(Context); break;
+				case EVectorVMOp::clampi: FVectorIntKernelClamp::Exec(Context); break;
+				case EVectorVMOp::mini: FVectorIntKernelMin::Exec(Context); break;
+				case EVectorVMOp::maxi: FVectorIntKernelMax::Exec(Context); break;
+				case EVectorVMOp::absi: FVectorIntKernelAbs::Exec(Context); break;
+				case EVectorVMOp::negi: FVectorIntKernelNegate::Exec(Context); break;
+				case EVectorVMOp::signi: FVectorIntKernelSign::Exec(Context); break;
+				case EVectorVMOp::randomi: FScalarIntKernelRandom::Exec(Context); break;
+				case EVectorVMOp::cmplti: FVectorIntKernelCompareLT::Exec(Context); break;
+				case EVectorVMOp::cmplei: FVectorIntKernelCompareLE::Exec(Context); break;
+				case EVectorVMOp::cmpgti: FVectorIntKernelCompareGT::Exec(Context); break;
+				case EVectorVMOp::cmpgei: FVectorIntKernelCompareGE::Exec(Context); break;
+				case EVectorVMOp::cmpeqi: FVectorIntKernelCompareEQ::Exec(Context); break;
+				case EVectorVMOp::cmpneqi: FVectorIntKernelCompareNEQ::Exec(Context); break;
+				case EVectorVMOp::bit_and: FVectorIntKernelBitAnd::Exec(Context); break;
+				case EVectorVMOp::bit_or: FVectorIntKernelBitOr::Exec(Context); break;
+				case EVectorVMOp::bit_xor: FVectorIntKernelBitXor::Exec(Context); break;
+				case EVectorVMOp::bit_not: FVectorIntKernelBitNot::Exec(Context); break;
+				case EVectorVMOp::logic_and: FVectorIntKernelLogicAnd::Exec(Context); break;
+				case EVectorVMOp::logic_or: FVectorIntKernelLogicOr::Exec(Context); break;
+				case EVectorVMOp::logic_xor: FVectorIntKernelLogicXor::Exec(Context); break;
+				case EVectorVMOp::logic_not: FVectorIntKernelLogicNot::Exec(Context); break;
+				case EVectorVMOp::f2i: FVectorKernelFloatToInt::Exec(Context); break;
+				case EVectorVMOp::i2f: FVectorKernelIntToFloat::Exec(Context); break;
+				case EVectorVMOp::f2b: FVectorKernelFloatToBool::Exec(Context); break;
+				case EVectorVMOp::b2f: FVectorKernelBoolToFloat::Exec(Context); break;
+				case EVectorVMOp::i2b: FVectorKernelIntToBool::Exec(Context); break;
+				case EVectorVMOp::b2i: FVectorKernelBoolToInt::Exec(Context); break;
+
+				case EVectorVMOp::outputdata_32bit:	FScalarKernelWriteOutputIndexed<int32>::Exec(Context);	break;
+				case EVectorVMOp::inputdata_32bit: FVectorKernelReadInput<int32>::Exec(Context); break;
+					//case EVectorVMOp::inputdata_32bit: FVectorKernelReadInput32::Exec(Context); break;
+				case EVectorVMOp::inputdata_noadvance_32bit: FVectorKernelReadInputNoAdvance<int32>::Exec(Context); break;
+				case EVectorVMOp::acquireindex:	
+				{
+					if (bParallel)
+					{
+						FScalarKernelAcquireCounterIndex<true>::Exec(Context);
+					}
+					else
+					{
+						FScalarKernelAcquireCounterIndex<false>::Exec(Context);
+					}
+
+				}break;
+				case EVectorVMOp::external_func_call: FKernelExternalFunctionCall::Exec(Context); break;
+
+				case EVectorVMOp::exec_index: FVectorKernelExecutionIndex::Exec(Context); break;
+
+				case EVectorVMOp::enter_stat_scope: FVectorKernelEnterStatScope::Exec(Context); break;
+				case EVectorVMOp::exit_stat_scope: FVectorKernelExitStatScope::Exec(Context); break;
+
+			//Special case ops to handle unique IDs but this can be written as generalized buffer operations. TODO!
+			case EVectorVMOp::update_id:	FScalarKernelUpdateID::Exec(Context); break;
+			case EVectorVMOp::acquire_id:	FScalarKernelAcquireID::Exec(Context); break;
+
+					// Execution always terminates with a "done" opcode.
+				case EVectorVMOp::done:
+					break;
+
+					// Opcode not recognized / implemented.
+				default:
+					UE_LOG(LogVectorVM, Fatal, TEXT("Unknown op code 0x%02x"), (uint32)Op);
+					return;//BAIL
+				}
+			} while (Op != EVectorVMOp::done);
 
-		InstancesLeft -= InstancesPerChunk;
-		++ChunkIdx;
-	}
+			InstancesLeft -= InstancesPerChunk;
+			++ChunkIdx;
+		}
+	};
+
+	ParallelFor(NumBatches, ExecChunkBatch, GbParallelVVM == 0 || !bParallel);
 
 	// write back data set access indices, so we know how much was written to each data set
 	for (int32 Idx = 0; Idx < DataSetMetaTable.Num(); Idx++)
 	{
 		DataSetMetaTable[Idx].DataSetAccessIndex = DataSetIndexTable[Idx];	
 	}
-
 }
 
 uint8 VectorVM::GetNumOpCodes()
@@ -1596,10 +1747,9 @@ uint8 VectorVM::GetNumOpCodes()
 #if WITH_EDITOR
 FString VectorVM::GetOpName(EVectorVMOp Op)
 {
-	static UEnum* EnumStateObj = FindObject<UEnum>(ANY_PACKAGE, TEXT("EVectorVMOp"), true);	
-	check(EnumStateObj);
+	check(g_VectorVMEnumStateObj);
 
-	FString OpStr = EnumStateObj->GetNameByValue((uint8)Op).ToString();
+	FString OpStr = g_VectorVMEnumStateObj->GetNameByValue((uint8)Op).ToString();
 	int32 LastIdx = 0;
 	OpStr.FindLastChar(TEXT(':'),LastIdx);
 	return OpStr.RightChop(LastIdx);
@@ -1607,10 +1757,9 @@ FString VectorVM::GetOpName(EVectorVMOp Op)
 
 FString VectorVM::GetOperandLocationName(EVectorVMOperandLocation Location)
 {
-	static UEnum* EnumStateObj = FindObject<UEnum>(ANY_PACKAGE, TEXT("EVectorVMOperandLocation"), true);
-	check(EnumStateObj);
+	check(g_VectorVMEnumOperandObj);
 
-	FString LocStr = EnumStateObj->GetNameByValue((uint8)Location).ToString();	
+	FString LocStr = g_VectorVMEnumOperandObj->GetNameByValue((uint8)Location).ToString();
 	int32 LastIdx = 0;
 	LocStr.FindLastChar(TEXT(':'), LastIdx);
 	return LocStr.RightChop(LastIdx);
diff --git a/Engine/Source/Runtime/VectorVM/Public/VectorVM.h b/Engine/Source/Runtime/VectorVM/Public/VectorVM.h
--- a/Engine/Source/Runtime/VectorVM/Public/VectorVM.h
+++ b/Engine/Source/Runtime/VectorVM/Public/VectorVM.h
@@ -4,6 +4,7 @@
 
 #include "CoreMinimal.h"
 #include "UObject/ObjectMacros.h"
+#include "Math/RandomStream.h"
 
 //TODO: move to a per platform header and have VM scale vectorization according to vector width.
 #define VECTOR_WIDTH (128)
@@ -137,10 +138,18 @@ enum class EVectorVMOp : uint8
 	enter_stat_scope,
 	exit_stat_scope,
 
+	//updates an ID in the ID table
+	update_id,
+	//acquires a new ID from the free list.
+	acquire_id,
+
 	NumOpcodes
 };
 
 
+//TODO: 
+//All of this stuff can be handled by the VM compiler rather than dirtying the VM code.
+//Some require RWBuffer like support.
 struct FDataSetMeta
 {
 	uint8 **InputRegisters;
@@ -149,12 +158,28 @@ struct FDataSetMeta
 	int32 DataSetAccessIndex;	// index for individual elements of this set
 	int32 DataSetOffset;		// offset in the register table
 
-	FDataSetMeta(uint32 DataSetSize, uint8 **Data = nullptr, uint8 InNumVariables = 0)
-		: InputRegisters(Data), NumVariables(InNumVariables), DataSetSizeInBytes(DataSetSize), DataSetAccessIndex(0), DataSetOffset(0)
-	{}
+	int32 InstanceOffset;		// offset of the first instance processed 
+	
+	TArray<int32>*RESTRICT IDTable;
+	TArray<int32>*RESTRICT FreeIDTable;
+
+	/** Number of free IDs in the FreeIDTable */
+	int32* NumFreeIDs;
+
+	/** MaxID used in this execution. */
+	int32* MaxUsedID;
 
-	FDataSetMeta()
-		: InputRegisters(nullptr), NumVariables(0), DataSetSizeInBytes(0), DataSetAccessIndex(0), DataSetOffset(0)
+	int32 IDAcquireTag;
+
+	FDataSetMeta(uint32 DataSetSize, uint8 **Data, uint8 InNumVariables, int32 InInstanceOffset, TArray<int32>* InIDTable, TArray<int32>* InFreeIDTable, int32* InNumFreeIDs, int32* InMaxUsedID, int32 InIDAcquireTag)
+		: InputRegisters(Data), NumVariables(InNumVariables), DataSetSizeInBytes(DataSetSize), DataSetAccessIndex(INDEX_NONE), DataSetOffset(0), InstanceOffset(InInstanceOffset)
+		, IDTable(InIDTable), FreeIDTable(InFreeIDTable), NumFreeIDs(InNumFreeIDs), MaxUsedID(InMaxUsedID), IDAcquireTag(InIDAcquireTag)
+	{
+	}
+
+	FDataSetMeta() 
+		: InputRegisters(nullptr), NumVariables(0), DataSetSizeInBytes(0), DataSetAccessIndex(INDEX_NONE), DataSetOffset(0), InstanceOffset(0)
+		, IDTable(nullptr), FreeIDTable(nullptr), NumFreeIDs(nullptr), MaxUsedID(nullptr), IDAcquireTag(0)
 	{}
 };
 
@@ -197,7 +222,6 @@ namespace VectorVM
 		FVMExternalFunction* ExternalFunctionTable,
 		void** UserPtrTable,
 		int32 NumInstances
-
 #if STATS
 		, const TArray<TStatId>& StatScopes
 #endif
@@ -207,8 +231,6 @@ namespace VectorVM
 } // namespace VectorVM
 
 
-
-
   /**
   * Context information passed around during VM execution.
   */
@@ -231,6 +253,9 @@ struct FVectorVMContext : TThreadSingleton<FVectorVMContext>
 	/** Start instance of current chunk. */
 	int32 StartInstance;
 
+	/** Array of meta data on data sets. TODO: This struct should be removed and all features it contains be handled by more general vm ops and the compiler's knowledge of offsets etc. */
+	FDataSetMeta*RESTRICT DataSetMetaTable;
+
 #if STATS
 	TArray<FCycleCounter, TInlineAllocator<64>> StatCounterStack;
 	const TArray<TStatId>* StatScopes;
@@ -239,6 +264,8 @@ struct FVectorVMContext : TThreadSingleton<FVectorVMContext>
 	TArray<uint8, TAlignedHeapAllocator<VECTOR_WIDTH_BYTES>> TempRegTable;
 	uint8 *RESTRICT RegisterTable[VectorVM::MaxRegisters];
 
+	FRandomStream RandStream;
+
 	FVectorVMContext();
 
 	void PrepareForExec(
@@ -251,7 +278,8 @@ struct FVectorVMContext : TThreadSingleton<FVectorVMContext>
 		int32 *InDataSetOffsetTable,
 		int32 InNumSecondaryDatasets,
 		FVMExternalFunction* InExternalFunctionTable,
-		void** InUserPtrTable
+		void** InUserPtrTable,
+		FDataSetMeta*RESTRICT InDataSetMetaTable
 #if STATS
 		, const TArray<TStatId>* InStatScopes
 #endif
@@ -313,6 +341,7 @@ struct FConstantHandler : public FConstantHandlerBase
 		, Constant(*((T*)(Context.ConstantTable + ConstantIndex)))
 	{}
 	FORCEINLINE const T& Get() { return Constant; }
+	FORCEINLINE const T& GetAndAdvance() { return Constant; }
 };
 
 
@@ -325,6 +354,7 @@ struct FDataSetOffsetHandler : FConstantHandlerBase
 		, Offset(Context.DataSetOffsetTable[ConstantIndex])
 	{}
 	FORCEINLINE const uint32 Get() { return Offset; }
+	FORCEINLINE const uint32 GetAndAdvance() { return Offset; }
 };
 
 
@@ -338,6 +368,7 @@ struct FConstantHandler<VectorRegister> : public FConstantHandlerBase
 		, Constant(VectorLoadFloat1(&Context.ConstantTable[ConstantIndex]))
 	{}
 	FORCEINLINE const VectorRegister Get() { return Constant; }
+	FORCEINLINE const VectorRegister GetAndAdvance() { return Constant; }
 };
 
 template<>
@@ -349,6 +380,7 @@ struct FConstantHandler<VectorRegisterInt> : public FConstantHandlerBase
 		, Constant(VectorIntLoad1(&Context.ConstantTable[ConstantIndex]))
 	{}
 	FORCEINLINE const VectorRegisterInt Get() { return Constant; }
+	FORCEINLINE const VectorRegisterInt GetAndAdvance() { return Constant; }
 };
 
 //////////////////////////////////////////////////////////////////////////
@@ -390,6 +422,8 @@ struct FRegisterHandler : public FRegisterHandlerBase
 	FORCEINLINE const T Get() { return *Register; }
 	FORCEINLINE T* GetDest() { return Register; }
 	FORCEINLINE void Advance() { ++Register; }
+	FORCEINLINE const T GetAndAdvance() { return *Register++; }
+	FORCEINLINE T* GetDestAndAdvance() { return Register++; }
 };
 
 template<> struct FRegisterHandler<VectorRegister> : public FRegisterHandlerBase
@@ -403,6 +437,7 @@ template<> struct FRegisterHandler<VectorRegister> : public FRegisterHandlerBase
 	FORCEINLINE const VectorRegister Get() { return VectorLoadAligned(Register); }
 	FORCEINLINE VectorRegister* GetDest() { return Register; }
 	FORCEINLINE void Advance() { ++Register; }
+	FORCEINLINE const VectorRegister  GetAndAdvance() { return *Register++; }
 };
 
 template<> struct FRegisterHandler<VectorRegisterInt> : public FRegisterHandlerBase
@@ -416,6 +451,7 @@ template<> struct FRegisterHandler<VectorRegisterInt> : public FRegisterHandlerB
 	FORCEINLINE const VectorRegisterInt Get() { return VectorIntLoadAligned(Register); }
 	FORCEINLINE VectorRegisterInt* GetDest() { return Register; }
 	FORCEINLINE void Advance() { ++Register; }
+	FORCEINLINE const VectorRegisterInt GetAndAdvance() { return *Register++; }
 };
 
 /** Handles writing to a register, advancing the pointer with each write. */
@@ -430,6 +466,7 @@ struct FRegisterDestHandler : public FRegisterHandlerBase
 	FORCEINLINE T* RESTRICT GetDest() { return Register; }
 	FORCEINLINE T GetValue() { return *Register; }
 	FORCEINLINE void Advance() { ++Register; }
+	FORCEINLINE T* GetDestAndAdvance() { return Register++; }
 };
 
 template<> struct FRegisterDestHandler<VectorRegister> : public FRegisterHandlerBase
@@ -443,6 +480,7 @@ template<> struct FRegisterDestHandler<VectorRegister> : public FRegisterHandler
 
 	FORCEINLINE VectorRegister* RESTRICT GetDest() { return Register; }
 	FORCEINLINE void Advance() { ++Register; }
+	FORCEINLINE VectorRegister* GetDestAndAdvance() { return Register++; }
 };
 
 template<> struct FRegisterDestHandler<VectorRegisterInt> : public FRegisterHandlerBase
@@ -456,5 +494,6 @@ template<> struct FRegisterDestHandler<VectorRegisterInt> : public FRegisterHand
 
 	FORCEINLINE VectorRegisterInt* RESTRICT GetDest() { return Register; }
 	FORCEINLINE void Advance() { ++Register; }
+	FORCEINLINE VectorRegisterInt* GetDestAndAdvance() { return Register++; }
 };
 
diff --git a/Engine/Source/Runtime/Windows/D3D11RHI/Private/D3D11GPUReadback.cpp b/Engine/Source/Runtime/Windows/D3D11RHI/Private/D3D11GPUReadback.cpp
new file mode 100644
--- /dev/null
+++ b/Engine/Source/Runtime/Windows/D3D11RHI/Private/D3D11GPUReadback.cpp
@@ -0,0 +1,32 @@
+// Copyright 1998-2018 Epic Games, Inc. All Rights Reserved.
+
+/*=============================================================================
+D3D11GPUReadback.cpp: Convenience function implementations for async GPU
+memory updates and readbacks
+=============================================================================*/
+
+#include "D3D11RHIPrivate.h"
+
+FStagingBufferRHIRef FD3D11DynamicRHI::RHICreateStagingBuffer()
+{
+	return new FD3D11StagingBuffer();
+}
+
+
+
+/*
+FRHIGPUMemoryReadback *RHIScheduleGPUMemoryReadback(FRHICommandList& CmdList, FVertexBufferRHIRef GPUBuffer, FName RequestName)
+{
+	FRHIStagingBuffer *StagingBuffer = RHICreateStagingBuffer(); //new FD3D11StagingBuffer();
+	FRHIGPUMemoryReadback *Readback = new FRHIGPUMemoryReadback(StagingBuffer, GPUBuffer, RequestName);
+	return Readback;
+}
+
+FRHIGPUMemoryUpdate *RHIScheduleGPUMemoryUpdate(FRHICommandList& CmdList, FVertexBufferRHIRef GPUBuffer, FName RequestName)
+{
+	FRHIStagingBuffer *StagingBuffer = RHICreateStagingBuffer(); //new FD3D11StagingBuffer();
+	FRHIGPUMemoryUpdate *Update = new FRHIGPUMemoryUpdate(StagingBuffer, GPUBuffer, RequestName);
+
+	return Update;
+}
+*/
diff --git a/Engine/Source/Runtime/Windows/D3D11RHI/Private/D3D11Query.cpp b/Engine/Source/Runtime/Windows/D3D11RHI/Private/D3D11Query.cpp
--- a/Engine/Source/Runtime/Windows/D3D11RHI/Private/D3D11Query.cpp
+++ b/Engine/Source/Runtime/Windows/D3D11RHI/Private/D3D11Query.cpp
@@ -6,6 +6,12 @@
 
 #include "D3D11RHIPrivate.h"
 
+
+FGPUFenceRHIRef FD3D11DynamicRHI::RHICreateGPUFence(const FName &Name)
+{
+	return new FD3D11GPUFence(Name);
+}
+
 void FD3D11DynamicRHI::RHIBeginOcclusionQueryBatch(uint32 NumQueriesInBatch)
 {
 	check(RequestedOcclusionQueriesInBatch == 0);
diff --git a/Engine/Source/Runtime/Windows/D3D11RHI/Private/D3D11RHIPrivate.h b/Engine/Source/Runtime/Windows/D3D11RHI/Private/D3D11RHIPrivate.h
--- a/Engine/Source/Runtime/Windows/D3D11RHI/Private/D3D11RHIPrivate.h
+++ b/Engine/Source/Runtime/Windows/D3D11RHI/Private/D3D11RHIPrivate.h
@@ -25,6 +25,7 @@ DECLARE_LOG_CATEGORY_EXTERN(LogD3D11RHI, Log, All);
 #include "D3D11Util.h"
 #include "D3D11State.h"
 #include "D3D11Resources.h"
+#include "D3D11GPUReadback.h"
 #include "D3D11Viewport.h"
 #include "D3D11ConstantBuffer.h"
 #include "D3D11StateCache.h"
@@ -444,7 +445,8 @@ public:
 	virtual class IRHICommandContext* RHIGetDefaultContext() final override;
 	virtual class IRHICommandContextContainer* RHIGetCommandContextContainer(int32 Index, int32 Num) final override;
 
-
+	virtual FGPUFenceRHIRef RHICreateGPUFence(const FName &Name) final override;
+	virtual FStagingBufferRHIRef RHICreateStagingBuffer() final override;
 
 	virtual void RHISetComputeShader(FComputeShaderRHIParamRef ComputeShader) final override;
 	virtual void RHIDispatchComputeShader(uint32 ThreadGroupCountX, uint32 ThreadGroupCountY, uint32 ThreadGroupCountZ) final override;
diff --git a/Engine/Source/Runtime/Windows/D3D11RHI/Public/D3D11GPUReadback.h b/Engine/Source/Runtime/Windows/D3D11RHI/Public/D3D11GPUReadback.h
new file mode 100644
--- /dev/null
+++ b/Engine/Source/Runtime/Windows/D3D11RHI/Public/D3D11GPUReadback.h
@@ -0,0 +1,40 @@
+// Copyright 1998-2018 Epic Games, Inc. All Rights Reserved.
+
+/*=============================================================================
+D3D11Resources.h: D3D resource RHI definitions.
+=============================================================================*/
+
+#pragma once
+
+/* D3D11 implementation of staging buffer class
+*/
+class FD3D11StagingBuffer : public FRHIStagingBuffer
+{
+public:
+	FD3D11StagingBuffer()
+	{
+	}
+
+	// copyresource, map, return ptr
+	virtual void *Lock(FVertexBufferRHIRef GPUBuffer, uint32 Offset, uint32 NumBytes, EResourceLockMode LockMode) override
+	{
+		check(!LastLockedBuffer);	// attempting to map another buffer to this staging buffer without unlocking the previous first?
+
+		MappedPtr = RHILockVertexBuffer(GPUBuffer, Offset, NumBytes, LockMode);
+		LastLockedBuffer = GPUBuffer;
+		return MappedPtr;
+	}
+
+	// unmap, free memory
+	virtual void Unlock() override
+	{
+		check(LastLockedBuffer);	// attempting to unlock without having locked first?
+
+		RHIUnlockVertexBuffer(LastLockedBuffer);
+		MappedPtr = nullptr;
+		LastLockedBuffer = nullptr;
+	}
+};
+
+
+
diff --git a/Engine/Source/Runtime/Windows/D3D11RHI/Public/D3D11Resources.h b/Engine/Source/Runtime/Windows/D3D11RHI/Public/D3D11Resources.h
--- a/Engine/Source/Runtime/Windows/D3D11RHI/Public/D3D11Resources.h
+++ b/Engine/Source/Runtime/Windows/D3D11RHI/Public/D3D11Resources.h
@@ -609,6 +609,55 @@ FORCEINLINE FD3D11TextureBase* GetD3D11TextureFromRHITexture(FRHITexture* Textur
 	return Result;
 }
 
+
+
+/* D3D11 implementation of generic GPU fence
+ * The D3D11 always returns false for Poll until the next frame from the frame the fence was inserted
+ * because D3D11.0 doesn't have a GPU/CPU sync object, we need to fake it
+*/
+class FD3D11GPUFence : public FRHIGPUFence
+{
+public:
+	FD3D11GPUFence(FName InName)
+		: FRHIGPUFence(InName)
+	{
+		Write();
+	}
+
+	~FD3D11GPUFence()
+	{
+
+	}
+
+	virtual bool Write()
+	{
+		InsertedFrameNumber = GFrameNumberRenderThread;
+		return true;
+	}
+
+	virtual bool Poll() const override
+	{
+		if (GFrameNumberRenderThread > InsertedFrameNumber )
+		{
+			return true;
+		}
+		return false;
+	}
+
+	virtual bool Wait(float TimeoutMs) const override
+	{
+		if (GFrameNumberRenderThread > InsertedFrameNumber )
+		{
+			return true;
+		}
+		return false;
+	}
+
+private:
+	uint32 InsertedFrameNumber;
+};
+
+
 /** D3D11 render query */
 class FD3D11RenderQuery : public FRHIRenderQuery
 {
