Merging //UE4/Dev-Main to Dev-Build (//UE4/Dev-Build)

diff --git a/Engine/Source/Runtime/Engine/Public/ParticleHelper.h b/Engine/Source/Runtime/Engine/Public/ParticleHelper.h
--- a/Engine/Source/Runtime/Engine/Public/ParticleHelper.h
+++ b/Engine/Source/Runtime/Engine/Public/ParticleHelper.h
@@ -242,6 +242,13 @@ enum EParticleStates
 /*-----------------------------------------------------------------------------
 	FParticlesStatGroup
 -----------------------------------------------------------------------------*/
+
+DECLARE_STATS_GROUP(TEXT("ParticlesOverview"), STATGROUP_ParticlesOverview, STATCAT_Advanced);
+DECLARE_CYCLE_STAT_EXTERN(TEXT("GT Total"), STAT_ParticlesOverview_GT, STATGROUP_ParticlesOverview, );
+DECLARE_CYCLE_STAT_EXTERN(TEXT("GT Concurrent Total"), STAT_ParticlesOverview_GT_CNC, STATGROUP_ParticlesOverview, );
+DECLARE_CYCLE_STAT_EXTERN(TEXT("RT Total"), STAT_ParticlesOverview_RT, STATGROUP_ParticlesOverview, );
+DECLARE_CYCLE_STAT_EXTERN(TEXT("RT Concurrent Total"), STAT_ParticlesOverview_RT_CNC, STATGROUP_ParticlesOverview, );
+
 DECLARE_DWORD_COUNTER_STAT_EXTERN(TEXT("Sprite Particles"),STAT_SpriteParticles,STATGROUP_Particles, );
 DECLARE_DWORD_COUNTER_STAT_EXTERN(TEXT("Sprite Ptcls Spawned"),STAT_SpriteParticlesSpawned,STATGROUP_Particles, );
 DECLARE_DWORD_COUNTER_STAT_EXTERN(TEXT("Sprite Ptcls Updated"),STAT_SpriteParticlesUpdated,STATGROUP_Particles, );
diff --git a/Engine/Source/Runtime/RHI/Public/DynamicRHI.h b/Engine/Source/Runtime/RHI/Public/DynamicRHI.h
--- a/Engine/Source/Runtime/RHI/Public/DynamicRHI.h
+++ b/Engine/Source/Runtime/RHI/Public/DynamicRHI.h
@@ -228,7 +228,7 @@ public:
 
     /**
      * Lock a staging buffer to read contents on the CPU that were written by the GPU, without having to stall.
-     * @discussion This function requires that you have issued an EnqueueStagedRead invocation and verified that the FRHIGPUFence has been signalled before calling.
+     * @discussion This function requires that you have issued an EnqueueStagedRead invocation and verified that the FRHIGPUFence has been signaled before calling.
      * @param StagingBuffer The buffer to lock.
      * @param Offset The offset in the buffer to return.
      * @param SizeRHI The length of the region in the buffer to lock.
@@ -238,13 +238,13 @@ public:
     
     /**
      * Unlock a staging buffer previously locked with RHILockStagingBuffer.
-     * @param StagingBuffer The buffer to lock.
+     * @param StagingBuffer The buffer that was previously locked.
      */
     virtual void RHIUnlockStagingBuffer(FStagingBufferRHIParamRef StagingBuffer);
 
     /**
      * Lock a staging buffer to read contents on the CPU that were written by the GPU, without having to stall.
-     * @discussion This function requires that you have issued an EnqueueStagedRead invocation and verified that the FRHIGPUFence has been signalled before calling.
+     * @discussion This function requires that you have issued an EnqueueStagedRead invocation and verified that the FRHIGPUFence has been signaled before calling.
      * @param RHICmdList The command-list to execute on or synchronize with.
      * @param StagingBuffer The buffer to lock.
      * @param Offset The offset in the buffer to return.
@@ -256,7 +256,7 @@ public:
     /**
      * Unlock a staging buffer previously locked with LockStagingBuffer_RenderThread.
      * @param RHICmdList The command-list to execute on or synchronize with.
-     * @param StagingBuffer The buffer to lock.
+     * @param StagingBuffer The buffer what was previously locked.
      */
     virtual void UnlockStagingBuffer_RenderThread(class FRHICommandListImmediate& RHICmdList, FStagingBufferRHIParamRef StagingBuffer);
 
diff --git a/Engine/Source/Runtime/RenderCore/Public/RendererInterface.h b/Engine/Source/Runtime/RenderCore/Public/RendererInterface.h
--- a/Engine/Source/Runtime/RenderCore/Public/RendererInterface.h
+++ b/Engine/Source/Runtime/RenderCore/Public/RendererInterface.h
@@ -23,6 +23,7 @@ class FSceneRenderTargets;
 class FSceneView;
 class FSceneViewFamily;
 class FSceneTexturesUniformParameters;
+class FGlobalDistanceFieldParameterData;
 struct FMeshBatch;
 struct FSynthBenchmarkResults;
 class IVirtualTextureSpace;
@@ -502,19 +503,27 @@ enum EDrawRectangleFlags
 	EDRF_UseTesselatedIndexBuffer
 };
 
-class FPostOpaqueRenderParameters
+class FPreSceneRenderValues
 {
 	public:
-		FIntRect ViewportRect;
-		FMatrix ViewMatrix;
-		FMatrix ProjMatrix;
-		FRHITexture2D* DepthTexture;
-		FRHITexture2D* NormalTexture;
-		FRHITexture2D* SmallDepthTexture;
-		FRHICommandListImmediate* RHICmdList;
-		FUniformBufferRHIParamRef ViewUniformBuffer;
-		TUniformBufferRef<FSceneTexturesUniformParameters> SceneTexturesUniformParams;
-		void* Uid; // A unique identifier for the view.
+		bool bUsesGlobalDistanceField = false;
+};
+DECLARE_MULTICAST_DELEGATE_OneParam(FPreSceneRenderDelegate, class FPreSceneRenderValues&);
+
+class FPostOpaqueRenderParameters
+{
+public:
+	FIntRect ViewportRect;
+	FMatrix ViewMatrix;
+	FMatrix ProjMatrix;
+	FRHITexture2D* DepthTexture;
+	FRHITexture2D* NormalTexture;
+	FRHITexture2D* SmallDepthTexture;
+	FRHICommandListImmediate* RHICmdList;
+	FUniformBufferRHIParamRef ViewUniformBuffer;
+	TUniformBufferRef<FSceneTexturesUniformParameters> SceneTexturesUniformParams;
+	const FGlobalDistanceFieldParameterData* GlobalDistanceFieldParams;
+	void* Uid; // A unique identifier for the view.
 };
 DECLARE_DELEGATE_OneParam(FPostOpaqueRenderDelegate, class FPostOpaqueRenderParameters&);
 
@@ -764,10 +773,12 @@ public:
 	virtual void RegisterCustomCullingImpl(ICustomCulling* impl) = 0;
 	virtual void UnregisterCustomCullingImpl(ICustomCulling* impl) = 0;
 
+	virtual FPreSceneRenderDelegate& OnPreSceneRender() = 0;
 	virtual void RegisterPostOpaqueRenderDelegate(const FPostOpaqueRenderDelegate& PostOpaqueRenderDelegate) = 0;
 	virtual void RegisterOverlayRenderDelegate(const FPostOpaqueRenderDelegate& OverlayRenderDelegate) = 0;
 	virtual void RenderPostOpaqueExtensions(const class FViewInfo& View, FRHICommandListImmediate& RHICmdList, class FSceneRenderTargets& SceneContext, TUniformBufferRef<FSceneTexturesUniformParameters>& SceneTextureUniformParams) = 0;
 	virtual void RenderOverlayExtensions(const class FViewInfo& View, FRHICommandListImmediate& RHICmdList, FSceneRenderTargets& SceneContext) = 0;
+	virtual FPreSceneRenderValues PreSceneRenderExtension() = 0;
 	virtual bool HasPostOpaqueExtentions() const = 0;
 
 	virtual void RegisterPostOpaqueComputeDispatcher(FComputeDispatcher *Dispatcher) = 0;
diff --git a/Engine/Source/Runtime/Renderer/Private/DeferredShadingRenderer.cpp b/Engine/Source/Runtime/Renderer/Private/DeferredShadingRenderer.cpp
--- a/Engine/Source/Runtime/Renderer/Private/DeferredShadingRenderer.cpp
+++ b/Engine/Source/Runtime/Renderer/Private/DeferredShadingRenderer.cpp
@@ -642,7 +642,9 @@ void FDeferredShadingSceneRenderer::Render(FRHICommandListImmediate& RHICmdList)
 		SCOPE_CYCLE_COUNTER(STAT_PostInitViews_FlushDel);
 		RHICmdList.ImmediateFlush(EImmediateFlushType::FlushRHIThreadFlushResources);
 	}
-
+	IRendererModule& RendererModule = GetRendererModule();
+	FPreSceneRenderValues PreSceneRenderValues = RendererModule.PreSceneRenderExtension();
+	Views[0].bUsesGlobalDistanceField |= PreSceneRenderValues.bUsesGlobalDistanceField;
 	if (!bDoInitViewAftersPrepass)
 	{
 		bool bSplitDispatch = !GDoPrepareDistanceFieldSceneAfterRHIFlush;
@@ -1461,7 +1463,6 @@ void FDeferredShadingSceneRenderer::Render(FRHICommandListImmediate& RHICmdList)
 
 	checkSlow(RHICmdList.IsOutsideRenderPass());
 
-	IRendererModule& RendererModule = GetRendererModule();
 	if (RendererModule.HasPostOpaqueExtentions())
 	{
 		FSceneTexturesUniformParameters SceneTextureParameters;
@@ -1478,6 +1479,7 @@ void FDeferredShadingSceneRenderer::Render(FRHICommandListImmediate& RHICmdList)
 		SceneContext.FinishRenderingSceneColor(RHICmdList);
 	}
 	checkSlow(RHICmdList.IsOutsideRenderPass());
+	UnbindRenderTargets(RHICmdList);
 	RendererModule.DispatchPostOpaqueCompute(RHICmdList, Views[0].ViewUniformBuffer);
 
 	// No longer needed, release
diff --git a/Engine/Source/Runtime/Renderer/Private/RendererModule.h b/Engine/Source/Runtime/Renderer/Private/RendererModule.h
--- a/Engine/Source/Runtime/Renderer/Private/RendererModule.h
+++ b/Engine/Source/Runtime/Renderer/Private/RendererModule.h
@@ -78,10 +78,12 @@ public:
 	virtual void UnRegisterPostOpaqueComputeDispatcher(FComputeDispatcher *Dispatcher) override;
 	virtual void DispatchPostOpaqueCompute(FRHICommandList &CmdList, FUniformBufferRHIParamRef ViewUniformBuffer) override;
 
+	virtual FPreSceneRenderDelegate& OnPreSceneRender() override { return PreSceneRenderDelegate; }
 	virtual void RegisterPostOpaqueRenderDelegate(const FPostOpaqueRenderDelegate& PostOpaqueRenderDelegate) override;
 	virtual void RegisterOverlayRenderDelegate(const FPostOpaqueRenderDelegate& OverlayRenderDelegate) override;
 	virtual void RenderPostOpaqueExtensions(const FViewInfo& View, FRHICommandListImmediate& RHICmdList, FSceneRenderTargets& SceneContext, TUniformBufferRef<FSceneTexturesUniformParameters>& SceneTextureUniformParams) override;
 	virtual void RenderOverlayExtensions(const FViewInfo& View, FRHICommandListImmediate& RHICmdList, FSceneRenderTargets& SceneContext) override;
+	virtual FPreSceneRenderValues PreSceneRenderExtension() override;
 
 	virtual bool HasPostOpaqueExtentions() const override
 	{
@@ -102,6 +104,7 @@ public:
 
 private:
 	TSet<FSceneInterface*> AllocatedScenes;
+	FPreSceneRenderDelegate PreSceneRenderDelegate;
 	FPostOpaqueRenderDelegate PostOpaqueRenderDelegate;
 	FPostOpaqueRenderDelegate OverlayRenderDelegate;
 	FOnResolvedSceneColor PostResolvedSceneColorCallbacks;
diff --git a/Engine/Source/Runtime/Renderer/Private/SceneRendering.cpp b/Engine/Source/Runtime/Renderer/Private/SceneRendering.cpp
--- a/Engine/Source/Runtime/Renderer/Private/SceneRendering.cpp
+++ b/Engine/Source/Runtime/Renderer/Private/SceneRendering.cpp
@@ -3261,6 +3261,13 @@ void FRendererModule::RegisterOverlayRenderDelegate(const FPostOpaqueRenderDeleg
 	this->OverlayRenderDelegate = InOverlayRenderDelegate;
 }
 
+FPreSceneRenderValues FRendererModule::PreSceneRenderExtension()
+{
+	FPreSceneRenderValues Result;
+	PreSceneRenderDelegate.Broadcast(Result);
+	return Result;
+}
+
 void FRendererModule::RenderPostOpaqueExtensions(const FViewInfo& View, FRHICommandListImmediate& RHICmdList, FSceneRenderTargets& SceneContext, TUniformBufferRef<FSceneTexturesUniformParameters>& SceneTextureUniformParams)
 {
 	check(IsInRenderingThread());
@@ -3272,6 +3279,7 @@ void FRendererModule::RenderPostOpaqueExtensions(const FViewInfo& View, FRHIComm
 	RenderParameters.SmallDepthTexture = SceneContext.GetSmallDepthSurface()->GetTexture2D();
 	RenderParameters.ViewUniformBuffer = View.ViewUniformBuffer;
 	RenderParameters.SceneTexturesUniformParams = SceneTextureUniformParams;
+	RenderParameters.GlobalDistanceFieldParams = &View.GlobalDistanceFieldInfo.ParameterData;
 
 	RenderParameters.ViewportRect = View.ViewRect;
 	RenderParameters.RHICmdList = &RHICmdList;
diff --git a/Engine/Source/Runtime/Renderer/Public/GlobalDistanceFieldParameters.h b/Engine/Source/Runtime/Renderer/Public/GlobalDistanceFieldParameters.h
--- a/Engine/Source/Runtime/Renderer/Public/GlobalDistanceFieldParameters.h
+++ b/Engine/Source/Runtime/Renderer/Public/GlobalDistanceFieldParameters.h
@@ -52,6 +52,11 @@ public:
 		MaxGlobalDistance.Bind(ParameterMap,TEXT("MaxGlobalDistance"));
 	}
 
+	bool IsBound() const
+	{
+		return GlobalVolumeCenterAndExtent.IsBound() || GlobalVolumeWorldToUVAddAndMul.IsBound();
+	}
+
 	friend FArchive& operator<<(FArchive& Ar,FGlobalDistanceFieldParameters& Parameters)
 	{
 		Ar << Parameters.GlobalDistanceFieldTexture0;
@@ -71,9 +76,9 @@ public:
 	}
 
 	template<typename ShaderRHIParamRef>
-	FORCEINLINE_DEBUGGABLE void Set(FRHICommandList& RHICmdList, const ShaderRHIParamRef ShaderRHI, const FGlobalDistanceFieldParameterData& ParameterData)
+	FORCEINLINE_DEBUGGABLE void Set(FRHICommandList& RHICmdList, const ShaderRHIParamRef ShaderRHI, const FGlobalDistanceFieldParameterData& ParameterData) const
 	{
-		if (GlobalVolumeCenterAndExtent.IsBound() || GlobalVolumeWorldToUVAddAndMul.IsBound())
+		if (IsBound())
 		{
 			SetTextureParameter(RHICmdList, ShaderRHI, GlobalDistanceFieldTexture0, GlobalDistanceFieldSampler0, TStaticSamplerState<SF_Bilinear, AM_Wrap, AM_Wrap, AM_Wrap>::GetRHI(), ParameterData.Textures[0] ? ParameterData.Textures[0] : GBlackVolumeTexture->TextureRHI.GetReference());
 			SetTextureParameter(RHICmdList, ShaderRHI, GlobalDistanceFieldTexture1, GlobalDistanceFieldSampler1, TStaticSamplerState<SF_Bilinear, AM_Wrap, AM_Wrap, AM_Wrap>::GetRHI(), ParameterData.Textures[1] ? ParameterData.Textures[1] : GBlackVolumeTexture->TextureRHI.GetReference());
diff --git a/Engine/Source/Runtime/VectorVM/Private/VectorVM.cpp b/Engine/Source/Runtime/VectorVM/Private/VectorVM.cpp
--- a/Engine/Source/Runtime/VectorVM/Private/VectorVM.cpp
+++ b/Engine/Source/Runtime/VectorVM/Private/VectorVM.cpp
@@ -13,7 +13,8 @@ IMPLEMENT_MODULE(FDefaultModuleImpl, VectorVM);
 
 
 DECLARE_STATS_GROUP(TEXT("VectorVM"), STATGROUP_VectorVM, STATCAT_Advanced);
-DECLARE_CYCLE_STAT(TEXT("Execution"), STAT_VVMExec, STATGROUP_VectorVM);
+DECLARE_CYCLE_STAT(TEXT("VVM Execution"), STAT_VVMExec, STATGROUP_VectorVM);
+DECLARE_CYCLE_STAT(TEXT("VVM Chunk"), STAT_VVMExecChunk, STATGROUP_VectorVM);
 
 DEFINE_LOG_CATEGORY_STATIC(LogVectorVM, All, All);
 
@@ -1375,6 +1376,30 @@ struct FVectorIntKernelMultiply : TBinaryVectorIntKernel<FVectorIntKernelMultipl
 	}
 };
 
+//divi,
+struct FVectorIntKernelDivide : TBinaryVectorIntKernel<FVectorIntKernelDivide>
+{
+	static void VM_FORCEINLINE DoKernel(FVectorVMContext& Context, VectorRegisterInt* Dst, VectorRegisterInt Src0, VectorRegisterInt Src1)
+	{
+		int32 TmpA[4];
+		VectorIntStore(Src0, TmpA);
+
+		int32 TmpB[4];
+		VectorIntStore(Src1, TmpB);
+
+		// No intrinsics exist for integer divide. Since div by zero causes crashes, we must be safe against that.
+
+		int32 TmpDst[4];
+		TmpDst[0] = TmpB[0] != 0 ? (TmpA[0] / TmpB[0]) : 0;
+		TmpDst[1] = TmpB[1] != 0 ? (TmpA[1] / TmpB[1]) : 0;
+		TmpDst[2] = TmpB[2] != 0 ? (TmpA[2] / TmpB[2]) : 0;
+		TmpDst[3] = TmpB[3] != 0 ? (TmpA[3] / TmpB[3]) : 0;
+
+		*Dst = MakeVectorRegisterInt(TmpDst[0], TmpDst[1], TmpDst[2], TmpDst[3]);
+	}
+};
+
+
 //clampi,
 struct FVectorIntKernelClamp : TTrinaryVectorIntKernel<FVectorIntKernelClamp>
 {
@@ -1531,6 +1556,46 @@ struct FVectorIntKernelBitNot : TUnaryVectorIntKernel<FVectorIntKernelBitNot>
 	}
 };
 
+// bit_lshift
+struct FVectorIntKernelBitLShift : TBinaryVectorIntKernel<FVectorIntKernelBitLShift>
+{
+	static void VM_FORCEINLINE DoKernel(FVectorVMContext& Context, VectorRegisterInt* Dst, VectorRegisterInt Src0,  VectorRegisterInt Src1)
+	{
+		int32 TmpA[4];
+		VectorIntStore(Src0, TmpA);
+
+		int32 TmpB[4];
+		VectorIntStore(Src1, TmpB);
+
+		int32 TmpDst[4];
+		TmpDst[0] = (TmpA[0] << TmpB[0]);
+		TmpDst[1] = (TmpA[1] << TmpB[1]);
+		TmpDst[2] = (TmpA[2] << TmpB[2]);
+		TmpDst[3] = (TmpA[3] << TmpB[3]);
+		*Dst = MakeVectorRegisterInt(TmpDst[0], TmpDst[1], TmpDst[2], TmpDst[3] );
+	}
+};
+
+// bit_rshift
+struct FVectorIntKernelBitRShift : TBinaryVectorIntKernel<FVectorIntKernelBitRShift>
+{
+	static void VM_FORCEINLINE DoKernel(FVectorVMContext& Context, VectorRegisterInt* Dst, VectorRegisterInt Src0, VectorRegisterInt Src1)
+	{
+		int32 TmpA[4];
+		VectorIntStore(Src0, TmpA);
+
+		int32 TmpB[4];
+		VectorIntStore(Src1, TmpB);
+
+		int32 TmpDst[4];
+		TmpDst[0] = (TmpA[0] >> TmpB[0]);
+		TmpDst[1] = (TmpA[1] >> TmpB[1]);
+		TmpDst[2] = (TmpA[2] >> TmpB[2]);
+		TmpDst[3] = (TmpA[3] >> TmpB[3]);
+		*Dst = MakeVectorRegisterInt(TmpDst[0], TmpDst[1], TmpDst[2], TmpDst[3]);
+	}
+};
+
 //"Boolean" ops. Currently handling bools as integers.
 //logic_and,
 struct FVectorIntKernelLogicAnd : TBinaryVectorIntKernel<FVectorIntKernelLogicAnd>
@@ -1749,12 +1814,14 @@ void VectorVM::Exec(
 	}
 
 	int32 NumChunks = (NumInstances / InstancesPerChunk) + 1;
-	int32 ChunksPerBatch = GbParallelVVM != 0 ? GParallelVVMChunksPerBatch : NumChunks;
+	int32 ChunksPerBatch = (GbParallelVVM != 0 && FApp::ShouldUseThreadingForPerformance()) ? GParallelVVMChunksPerBatch : NumChunks;
 	int32 NumBatches = FMath::DivideAndRoundUp(NumChunks, ChunksPerBatch);
 	bool bParallel = NumBatches > 1;
 
 	auto ExecChunkBatch = [&](int32 BatchIdx)
-	{		
+	{
+		SCOPE_CYCLE_COUNTER(STAT_VVMExecChunk);
+
 		FVectorVMContext& Context = FVectorVMContext::Get();
 		Context.PrepareForExec(InputRegisters, OutputRegisters, NumInputRegisters, NumOutputRegisters, ConstantTable, DataSetIndexTable.GetData(), DataSetOffsetTable.GetData(), DataSetOffsetTable.Num(),
 			ExternalFunctionTable, UserPtrTable, DataSetMetaTable
@@ -1834,6 +1901,7 @@ void VectorVM::Exec(
 				case EVectorVMOp::addi: FVectorIntKernelAdd::Exec(Context); break;
 				case EVectorVMOp::subi: FVectorIntKernelSubtract::Exec(Context); break;
 				case EVectorVMOp::muli: FVectorIntKernelMultiply::Exec(Context); break;
+				case EVectorVMOp::divi: FVectorIntKernelDivide::Exec(Context); break;
 				case EVectorVMOp::clampi: FVectorIntKernelClamp::Exec(Context); break;
 				case EVectorVMOp::mini: FVectorIntKernelMin::Exec(Context); break;
 				case EVectorVMOp::maxi: FVectorIntKernelMax::Exec(Context); break;
@@ -1851,6 +1919,8 @@ void VectorVM::Exec(
 				case EVectorVMOp::bit_or: FVectorIntKernelBitOr::Exec(Context); break;
 				case EVectorVMOp::bit_xor: FVectorIntKernelBitXor::Exec(Context); break;
 				case EVectorVMOp::bit_not: FVectorIntKernelBitNot::Exec(Context); break;
+				case EVectorVMOp::bit_lshift: FVectorIntKernelBitLShift::Exec(Context); break;
+				case EVectorVMOp::bit_rshift: FVectorIntKernelBitRShift::Exec(Context); break;
 				case EVectorVMOp::logic_and: FVectorIntKernelLogicAnd::Exec(Context); break;
 				case EVectorVMOp::logic_or: FVectorIntKernelLogicOr::Exec(Context); break;
 				case EVectorVMOp::logic_xor: FVectorIntKernelLogicXor::Exec(Context); break;
diff --git a/Engine/Source/Runtime/VectorVM/Public/VectorVM.h b/Engine/Source/Runtime/VectorVM/Public/VectorVM.h
--- a/Engine/Source/Runtime/VectorVM/Public/VectorVM.h
+++ b/Engine/Source/Runtime/VectorVM/Public/VectorVM.h
@@ -87,7 +87,7 @@ enum class EVectorVMOp : uint8
 	addi,
 	subi,
 	muli,
-	//divi,//SSE Integer division is not implemented as an intrinsic. Will have to do some manual implementation.
+	divi,//SSE Integer division is not implemented as an intrinsic. Will have to do some manual implementation.
 	clampi,
 	mini,
 	maxi,
@@ -105,6 +105,8 @@ enum class EVectorVMOp : uint8
 	bit_or,
 	bit_xor,
 	bit_not,
+	bit_lshift,
+	bit_rshift,
 
 	//"Boolean" ops. Currently handling bools as integers.
 	logic_and,
