Pulling over Simon's ExecState fix CL#6675144
Fix for ExecutionState crash with many emitters.

diff --git a/Engine/Source/Runtime/VectorVM/Private/VectorVM.cpp b/Engine/Source/Runtime/VectorVM/Private/VectorVM.cpp
--- a/Engine/Source/Runtime/VectorVM/Private/VectorVM.cpp
+++ b/Engine/Source/Runtime/VectorVM/Private/VectorVM.cpp
@@ -1813,7 +1813,7 @@ void VectorVM::Exec(
 	//
 	for (int32 Idx = 0; Idx < DataSetMetaTable.Num(); Idx++)
 	{
-		uint32 DataSetOffset = DataSetMetaTable[Idx].NumVariables;
+		uint32 DataSetOffset = DataSetMetaTable[Idx].RegisterOffset;
 		DataSetOffsetTable.Add(DataSetOffset);
 		DataSetIndexTable.Add(DataSetMetaTable[Idx].DataSetAccessIndex);	// prime counter index table with the data set offset; will be incremented with every write for each instance
 	}
diff --git a/Engine/Source/Runtime/VectorVM/Public/VectorVM.h b/Engine/Source/Runtime/VectorVM/Public/VectorVM.h
--- a/Engine/Source/Runtime/VectorVM/Public/VectorVM.h
+++ b/Engine/Source/Runtime/VectorVM/Public/VectorVM.h
@@ -155,7 +155,7 @@ enum class EVectorVMOp : uint8
 struct FDataSetMeta
 {
 	uint8 **InputRegisters;
-	uint8 NumVariables;
+	int32 RegisterOffset;
 	int32 DataSetAccessIndex;	// index for individual elements of this set
 
 	int32 InstanceOffset;		// offset of the first instance processed 
@@ -180,14 +180,14 @@ struct FDataSetMeta
 	FORCEINLINE void LockFreeTable();
 	FORCEINLINE void UnlockFreeTable();
 
-	FDataSetMeta(uint8 **Data, uint8 InNumVariables, int32 InInstanceOffset, TArray<int32>* InIDTable, TArray<int32>* InFreeIDTable, int32* InNumFreeIDs, int32* InMaxUsedID, int32 InIDAcquireTag)
-		: InputRegisters(Data), NumVariables(InNumVariables), DataSetAccessIndex(INDEX_NONE), InstanceOffset(InInstanceOffset)
+	FDataSetMeta(uint8 **Data, int32 InRegisterOffset, int32 InInstanceOffset, TArray<int32>* InIDTable, TArray<int32>* InFreeIDTable, int32* InNumFreeIDs, int32* InMaxUsedID, int32 InIDAcquireTag)
+		: InputRegisters(Data), RegisterOffset(InRegisterOffset), DataSetAccessIndex(INDEX_NONE), InstanceOffset(InInstanceOffset)
 		, IDTable(InIDTable), FreeIDTable(InFreeIDTable), NumFreeIDs(InNumFreeIDs), MaxUsedID(InMaxUsedID), IDAcquireTag(InIDAcquireTag)
 	{
 	}
 
 	FDataSetMeta() 
-		: InputRegisters(nullptr), NumVariables(0), DataSetAccessIndex(INDEX_NONE), InstanceOffset(0)
+		: InputRegisters(nullptr), RegisterOffset(0), DataSetAccessIndex(INDEX_NONE), InstanceOffset(0)
 		, IDTable(nullptr), FreeIDTable(nullptr), NumFreeIDs(nullptr), MaxUsedID(nullptr), IDAcquireTag(0)
 	{}
 
