Niagara - Fix Win32 compile issue caused by cl 4418639

diff --git a/Engine/Source/Runtime/VectorVM/Public/VectorVM.h b/Engine/Source/Runtime/VectorVM/Public/VectorVM.h
--- a/Engine/Source/Runtime/VectorVM/Public/VectorVM.h
+++ b/Engine/Source/Runtime/VectorVM/Public/VectorVM.h
@@ -357,7 +357,7 @@ namespace VectorVM
 		{}
 
 		FORCEINLINE bool IsConstant()const { return !IsRegister(); }
-		FORCEINLINE bool IsRegister()const { return InputOffset & VVM_EXT_FUNC_INPUT_LOC_BIT; }
+		FORCEINLINE bool IsRegister()const { return (InputOffset & VVM_EXT_FUNC_INPUT_LOC_BIT) != 0; }
 		FORCEINLINE int32 GetOffset()const { return InputOffset & VVM_EXT_FUNC_INPUT_LOC_MASK; }
 
 		FORCEINLINE const T Get() { return *InputPtr; }
