Merging from Dev-Niagara 4427076

Fix for static analysis warnings.

diff --git a/Engine/Source/Runtime/VectorVM/Public/VectorVM.h b/Engine/Source/Runtime/VectorVM/Public/VectorVM.h
--- a/Engine/Source/Runtime/VectorVM/Public/VectorVM.h
+++ b/Engine/Source/Runtime/VectorVM/Public/VectorVM.h
@@ -345,9 +345,9 @@ namespace VectorVM
 	{
 	private:
 		/** Either byte offset into constant table or offset into register table deepening on VVM_INPUT_LOCATION_BIT */
-		uint32 InputOffset;
+		int32 InputOffset;
 		T* RESTRICT InputPtr;
-		uint32 AdvanceOffset;
+		int32 AdvanceOffset;
 
 	public:
 		FORCEINLINE FExternalFuncInputHandler(FVectorVMContext& Context)
@@ -381,16 +381,26 @@ namespace VectorVM
 	struct FExternalFuncRegisterHandler
 	{
 	private:
-		uint32 RegisterIndex;
-		uint32 AdvanceOffset;
+		int32 RegisterIndex;
+		int32 AdvanceOffset;
 		T Dummy;
 		T* RESTRICT Register;
 	public:
 		FORCEINLINE FExternalFuncRegisterHandler(FVectorVMContext& Context)
 			: RegisterIndex(DecodeU16(Context) & VVM_EXT_FUNC_INPUT_LOC_MASK)
 			, AdvanceOffset(IsValid() ? 1 : 0)
-			, Register(IsValid() ? (T*)Context.RegisterTable[RegisterIndex] : &Dummy)
-		{}
+		{
+			if (IsValid())
+			{
+				check(RegisterIndex < VectorVM::MaxRegisters);
+				Register = (T*)Context.RegisterTable[RegisterIndex];
+			}
+			else
+			{
+				Register = &Dummy;
+			}
+		}
+
 		FORCEINLINE bool IsValid() const { return RegisterIndex != (uint16)VVM_EXT_FUNC_INPUT_LOC_MASK; }
 
 		FORCEINLINE const T Get() { return *Register; }
