Fix for VectorVM race condition crash.

Adding a lock around places that modify the free table.
Keeping thread local array of IDs to free to minimize locks.
Modified free table allocation. Previously size could be insuffiicient in some cases.

diff --git a/Engine/Plugins/FX/Niagara/Source/Niagara/Classes/NiagaraDataSet.h b/Engine/Plugins/FX/Niagara/Source/Niagara/Classes/NiagaraDataSet.h
--- a/Engine/Plugins/FX/Niagara/Source/Niagara/Classes/NiagaraDataSet.h
+++ b/Engine/Plugins/FX/Niagara/Source/Niagara/Classes/NiagaraDataSet.h
@@ -306,16 +306,15 @@ public:
 
 			int32 RequiredIDs = FMath::Max(NumInstances, NumUsedIDs);
 			int32 ExistingNumIDs = PrevIDTable().Num();
-
-			//Free ID Table must always be at least as large as the data buffer + it's current size in the case all particles die this frame.
-			FreeIDsTable.SetNumUninitialized(NumInstances + NumFreeIDs);
+			int32 NumNewIDs = RequiredIDs - ExistingNumIDs;
 
 			if (RequiredIDs > ExistingNumIDs)
 			{
 				//UE_LOG(LogNiagara, Warning, TEXT("Growing ID Table! OldSize:%d | NewSize:%d"), ExistingNumIDs, RequiredIDs);
 				IDToIndexTable[CurrBuffer].SetNumUninitialized(RequiredIDs);
 
-				int32 NumNewIDs = RequiredIDs - ExistingNumIDs;
+				//Free ID Table must always be at least as large as the data buffer + it's current size in the case all particles die this frame.
+				FreeIDsTable.AddUninitialized(NumNewIDs);
 
 				//Free table should always have enough room for these new IDs.
 				check(NumFreeIDs + NumNewIDs <= FreeIDsTable.Num());
@@ -347,6 +346,9 @@ public:
 						++CheckedFreeID;
 					}
 				}
+
+				check(NumFreeIDs <= RequiredIDs);
+				FreeIDsTable.SetNumUninitialized(NumFreeIDs);
 			}
 			else
 			{
diff --git a/Engine/Plugins/FX/Niagara/Source/Niagara/Private/NiagaraScriptExecutionContext.cpp b/Engine/Plugins/FX/Niagara/Source/Niagara/Private/NiagaraScriptExecutionContext.cpp
--- a/Engine/Plugins/FX/Niagara/Source/Niagara/Private/NiagaraScriptExecutionContext.cpp
+++ b/Engine/Plugins/FX/Niagara/Source/Niagara/Private/NiagaraScriptExecutionContext.cpp
@@ -152,9 +152,8 @@ bool FNiagaraScriptExecutionContext::Execute(uint32 NumInstances, TArray<FNiagar
 			DataSetInfo.DataSet->CheckForNaNs();
 #endif
 			check(DataSetInfo.DataSet);
-			FDataSetMeta SetMeta(DataSetInfo.DataSet->GetSizeBytes(), &InputRegisters[NumInputRegisters], NumInputRegisters, DataSetInfo.StartInstance, 
+			DataSetMetaTable.Emplace(DataSetInfo.DataSet->GetSizeBytes(), &InputRegisters[NumInputRegisters], NumInputRegisters, DataSetInfo.StartInstance,
 				&DataSetInfo.DataSet->CurrIDTable(), &DataSetInfo.DataSet->GetFreeIDTable(), &DataSetInfo.DataSet->GetNumFreeIDs(), &DataSetInfo.DataSet->GetMaxUsedID(), DataSetInfo.DataSet->GetIDAcquireTag());
-			DataSetMetaTable.Add(SetMeta);
 			if (DataSetInfo.bAllocate)
 			{
 				DataSetInfo.DataSet->Allocate(NumInstances);
diff --git a/Engine/Source/Runtime/VectorVM/Private/VectorVM.cpp b/Engine/Source/Runtime/VectorVM/Private/VectorVM.cpp
--- a/Engine/Source/Runtime/VectorVM/Private/VectorVM.cpp
+++ b/Engine/Source/Runtime/VectorVM/Private/VectorVM.cpp
@@ -17,6 +17,12 @@ DECLARE_CYCLE_STAT(TEXT("Execution"), STAT_VVMExec, STATGROUP_VectorVM);
 
 DEFINE_LOG_CATEGORY_STATIC(LogVectorVM, All, All);
 
+//#define FREE_TABLE_LOCK_CONTENTION_WARNINGS (!UE_BUILD_SHIPPING)
+#define FREE_TABLE_LOCK_CONTENTION_WARNINGS (0)
+
+//I don't expect us to ever be waiting long
+#define FREE_TABLE_LOCK_CONTENTION_WARN_THRESHOLD_MS (0.01)
+
 //#define VM_FORCEINLINE
 #define VM_FORCEINLINE FORCEINLINE
 
@@ -34,6 +40,30 @@ DEFINE_LOG_CATEGORY_STATIC(LogVectorVM, All, All);
 #define SRCOP_CCR (OP2_CONST | OP1_CONST | OP_REGISTER)
 #define SRCOP_CCC (OP2_CONST | OP1_CONST | OP0_CONST)
 
+//Temporarily locking the free table until we can implement a lock free algorithm. UE-65856
+FORCEINLINE void FDataSetMeta::LockFreeTable()
+{
+#if FREE_TABLE_LOCK_CONTENTION_WARNINGS
+	uint64 StartCycles = FPlatformTime::Cycles64();
+#endif
+ 		
+	FreeTableLock.Lock();
+ 
+#if FREE_TABLE_LOCK_CONTENTION_WARNINGS
+	uint64 EndCylces = FPlatformTime::Cycles64();
+	double DurationMs = FPlatformTime::ToMilliseconds64(EndCylces - StartCycles);
+	if (DurationMs >= FREE_TABLE_LOCK_CONTENTION_WARN_THRESHOLD_MS)
+	{
+		UE_LOG(LogVectorVM, Warning, TEXT("VectorVM Stalled in LockFreeTable()! %g ms"), DurationMs);
+	}
+#endif
+}
+
+FORCEINLINE void FDataSetMeta::UnlockFreeTable()
+{
+ 	FreeTableLock.Unlock();
+}
+
 static int32 GbParallelVVM = 1;
 static FAutoConsoleVariableRef CVarbParallelVVM(
 	TEXT("vm.Parallel"),
@@ -159,7 +189,6 @@ FVectorVMContext::FVectorVMContext()
 	, UserPtrTable(nullptr)
 	, NumInstances(0)
 	, StartInstance(0)
-	, DataSetMetaTable(nullptr)
 #if STATS
 	, StatScopes(nullptr)
 #endif
@@ -186,7 +215,7 @@ void FVectorVMContext::PrepareForExec(
 	int32 InNumSecondaryDatasets,
 	FVMExternalFunction* InExternalFunctionTable,
 	void** InUserPtrTable,
-	FDataSetMeta*RESTRICT InDataSetMetaTable
+	TArray<FDataSetMeta>& RESTRICT InDataSetMetaTable
 #if STATS
 	, const TArray<TStatId>* InStatScopes
 #endif
@@ -214,7 +243,53 @@ void FVectorVMContext::PrepareForExec(
 		RegisterTable[VectorVM::NumTempRegisters + VectorVM::MaxInputRegisters + i] = OutputRegisters[i];
 	}
 
-	DataSetMetaTable = InDataSetMetaTable;
+	DataSetMetaTable = &InDataSetMetaTable;
+
+	ThreadLocalTempData.SetNum(DataSetMetaTable->Num());
+}
+
+void FVectorVMContext::FinishExec()
+{
+	//At the end of executing each chunk we can push any thread local temporary data out to the main storage with locks or atomics.
+
+	TArray<FDataSetMeta>& MetaTable = *DataSetMetaTable;
+	check(ThreadLocalTempData.Num() == MetaTable.Num());
+	for(int32 DataSetIndex=0; DataSetIndex < MetaTable.Num(); ++DataSetIndex)
+	{
+		FDataSetThreadLocalTempData&RESTRICT Data = ThreadLocalTempData[DataSetIndex];
+
+		if (Data.IDsToFree.Num() > 0)
+		{
+			TArray<int32>&RESTRICT FreeIDTable = *MetaTable[DataSetIndex].FreeIDTable;
+			int32&RESTRICT NumFreeIDs = *MetaTable[DataSetIndex].NumFreeIDs;
+			check(FreeIDTable.Num() >= NumFreeIDs + Data.IDsToFree.Num());
+
+			//Temporarily locking the free table until we can implement something lock-free
+			MetaTable[DataSetIndex].LockFreeTable();
+			for (int32 IDToFree : Data.IDsToFree)
+			{
+				FreeIDTable[NumFreeIDs++] = IDToFree;
+			}
+			//Unlock the free table.
+			MetaTable[DataSetIndex].UnlockFreeTable();
+			Data.IDsToFree.Reset();
+		}
+
+		//Also update the max ID seen. This should be the ONLY place in the VM we update this max value.
+		volatile int32* MaxUsedID = MetaTable[DataSetIndex].MaxUsedID;
+		int32 LocalMaxUsedID;
+		do
+		{
+			LocalMaxUsedID = *MaxUsedID;
+			if (LocalMaxUsedID >= Data.MaxID)
+			{
+				break;
+			}
+		} while (FPlatformAtomics::InterlockedCompareExchange(MaxUsedID, Data.MaxID, LocalMaxUsedID) != LocalMaxUsedID);
+
+		*MaxUsedID = FMath::Max(*MaxUsedID, Data.MaxID);
+
+	}
 }
 
 //////////////////////////////////////////////////////////////////////////
@@ -982,9 +1057,10 @@ struct FScalarKernelAcquireID
 	static VM_FORCEINLINE void Exec(FVectorVMContext& Context)
 	{
 		int32 DataSetIndex = VectorVM::DecodeU16(Context);
-		TArray<int32>&RESTRICT FreeIDTable = *Context.DataSetMetaTable[DataSetIndex].FreeIDTable;
+		TArray<FDataSetMeta>& MetaTable = *Context.DataSetMetaTable;
+		TArray<int32>&RESTRICT FreeIDTable = *MetaTable[DataSetIndex].FreeIDTable;
 
-		int32 Tag = Context.DataSetMetaTable[DataSetIndex].IDAcquireTag;
+		int32 Tag = MetaTable[DataSetIndex].IDAcquireTag;
 
 		int32 IDIndexReg = VectorVM::DecodeU16(Context);
 		int32*RESTRICT IDIndex = (int32*)(Context.RegisterTable[IDIndexReg]);
@@ -992,19 +1068,31 @@ struct FScalarKernelAcquireID
 		int32 IDTagReg = VectorVM::DecodeU16(Context);
 		int32*RESTRICT IDTag = (int32*)(Context.RegisterTable[IDTagReg]);
 
-		int32* NumFreeIDs = Context.DataSetMetaTable[DataSetIndex].NumFreeIDs;
+		int32& NumFreeIDs = *MetaTable[DataSetIndex].NumFreeIDs;
+
+		//Temporarily using a lock to ensure thread safety for accessing the FreeIDTable until a lock free solution can be implemented.
+		MetaTable[DataSetIndex].LockFreeTable();
 	
 		check(FreeIDTable.Num() >= Context.NumInstances);
+		check(NumFreeIDs >= Context.NumInstances);
 		for (int32 i = 0; i < Context.NumInstances; ++i)
 		{
-			int32 FreeIDTableIndex = FPlatformAtomics::InterlockedDecrement(NumFreeIDs);
-			//UE_LOG(LogVectorVM, Warning, TEXT("AcquireID: InstanceID:%d | FreeTableIdx:%d | Tag:%d "), FreeIDTable[FreeIDTableIndex], FreeIDTableIndex, Tag);
-			*IDIndex = FreeIDTable[FreeIDTableIndex];
+			int32 FreeIDTableIndex = --NumFreeIDs;
+
+			//Grab the value from the FreeIDTable.
+			int32 AcquiredID = FreeIDTable[FreeIDTableIndex];
+			checkSlow(AcquiredID != INDEX_NONE);
+
+			//Mark this entry in the FreeIDTable as invalid.
 			FreeIDTable[FreeIDTableIndex] = INDEX_NONE;
+
+			*IDIndex = AcquiredID;
 			*IDTag = Tag;
 			++IDIndex;
 			++IDTag;
 		}
+
+		MetaTable[DataSetIndex].UnlockFreeTable();
 	}
 };
 
@@ -1017,16 +1105,17 @@ struct FScalarKernelUpdateID
 		int32 InstanceIDRegisterIndex = VectorVM::DecodeU16(Context);
 		int32 InstanceIndexRegisterIndex = VectorVM::DecodeU16(Context);
 
-		TArray<int32>&RESTRICT IDTable = *Context.DataSetMetaTable[DataSetIndex].IDTable;
-		TArray<int32>&RESTRICT FreeIDTable = *Context.DataSetMetaTable[DataSetIndex].FreeIDTable;
-		int32 InstanceOffset = Context.DataSetMetaTable[DataSetIndex].InstanceOffset;// +Context.StartInstance;
+		TArray<FDataSetMeta>& MetaTable = *Context.DataSetMetaTable;
+
+		TArray<int32>&RESTRICT IDTable = *MetaTable[DataSetIndex].IDTable;
+		int32 InstanceOffset = MetaTable[DataSetIndex].InstanceOffset + Context.StartInstance;
 
 		int32*RESTRICT IDRegister = (int32*)(Context.RegisterTable[InstanceIDRegisterIndex]);
 		int32*RESTRICT IndexRegister = (int32*)(Context.RegisterTable[InstanceIndexRegisterIndex]);
+		
+		FDataSetThreadLocalTempData& DataSetTempData = Context.ThreadLocalTempData[DataSetIndex];
 
-		volatile int32* MaxUsedID = Context.DataSetMetaTable[DataSetIndex].MaxUsedID;
-		int32* NumFreeIDs = Context.DataSetMetaTable[DataSetIndex].NumFreeIDs;
-
+		TArray<int32>&RESTRICT IDsToFree = DataSetTempData.IDsToFree;
 		check(IDTable.Num() >= InstanceOffset + Context.NumInstances);
 		for (int32 i = 0; i < Context.NumInstances; ++i)
 		{
@@ -1035,28 +1124,17 @@ struct FScalarKernelUpdateID
 
 			if (Index == INDEX_NONE)
 			{
-				//Instance not being written so release it's ID.
-				IDTable[InstanceId] = INDEX_NONE;
-				int32 FreeIDIndex = FPlatformAtomics::InterlockedIncrement(NumFreeIDs) - 1;
-				//UE_LOG(LogVectorVM, Warning, TEXT("UpdateID: Dead instance, Adding ID %d to free table at idx %d."), InstanceId, FreeIDIndex);
-				FreeIDTable[FreeIDIndex] = InstanceId;
+				//Add the ID to a thread local list of IDs to free which are actually added to the list safely at the end of this chunk's execution.
+				IDsToFree.Add(InstanceId);
 			}
 			else
 			{
-				int32 RealIdx = InstanceOffset + Index;
+				//Update the actual index for this ID. No thread safety is needed as this ID slot can only ever be written by this instance and so a single thread.
+				int32 RealIdx = InstanceOffset + Index;	
 				IDTable[InstanceId] = RealIdx;
 
-				//Update the MaxUsedID.
-				int32 LocalMaxUsedID;
-				do 
-				{
-					LocalMaxUsedID = *MaxUsedID;
-					if (LocalMaxUsedID >= InstanceId)
-					{
-						break;
-					}
-				} 
-				while (FPlatformAtomics::InterlockedCompareExchange(MaxUsedID, InstanceId, LocalMaxUsedID) != LocalMaxUsedID);
+				//Update thread local max ID seen. We push this to the real value at the end of execution.
+				DataSetTempData.MaxID = FMath::Max(DataSetTempData.MaxID, InstanceId);
 				
 				//UE_LOG(LogVectorVM, Warning, TEXT("UpdateID: RealIdx:%d | InstanceID:%d."), RealIdx, InstanceId);
 			}
@@ -1679,7 +1757,7 @@ void VectorVM::Exec(
 	{		
 		FVectorVMContext& Context = FVectorVMContext::Get();
 		Context.PrepareForExec(InputRegisters, OutputRegisters, NumInputRegisters, NumOutputRegisters, ConstantTable, DataSetIndexTable.GetData(), DataSetOffsetTable.GetData(), DataSetOffsetTable.Num(),
-			ExternalFunctionTable, UserPtrTable, DataSetMetaTable.GetData()
+			ExternalFunctionTable, UserPtrTable, DataSetMetaTable
 #if STATS
 			, &StatScopes
 #endif
@@ -1825,6 +1903,8 @@ void VectorVM::Exec(
 			InstancesLeft -= InstancesPerChunk;
 			++ChunkIdx;
 		}
+
+		Context.FinishExec();
 	};
 
 	ParallelFor(NumBatches, ExecChunkBatch, GbParallelVVM == 0 || !bParallel);
diff --git a/Engine/Source/Runtime/VectorVM/Public/VectorVM.h b/Engine/Source/Runtime/VectorVM/Public/VectorVM.h
--- a/Engine/Source/Runtime/VectorVM/Public/VectorVM.h
+++ b/Engine/Source/Runtime/VectorVM/Public/VectorVM.h
@@ -171,6 +171,15 @@ struct FDataSetMeta
 
 	int32 IDAcquireTag;
 
+	//Temporary lock we're using for thread safety when writing to the FreeIDTable.
+	//TODO: A lock free algorithm is possible here. We can create a specialized lock free list and reuse the IDTable slots for FreeIndices as Next pointers for our LFL.
+	//This would also work well on the GPU. 
+	//UE-65856 for tracking this work.
+	FCriticalSection FreeTableLock;
+
+	FORCEINLINE void LockFreeTable();
+	FORCEINLINE void UnlockFreeTable();
+
 	FDataSetMeta(uint32 DataSetSize, uint8 **Data, uint8 InNumVariables, int32 InInstanceOffset, TArray<int32>* InIDTable, TArray<int32>* InFreeIDTable, int32* InNumFreeIDs, int32* InMaxUsedID, int32 InIDAcquireTag)
 		: InputRegisters(Data), NumVariables(InNumVariables), DataSetSizeInBytes(DataSetSize), DataSetAccessIndex(INDEX_NONE), DataSetOffset(0), InstanceOffset(InInstanceOffset)
 		, IDTable(InIDTable), FreeIDTable(InFreeIDTable), NumFreeIDs(InNumFreeIDs), MaxUsedID(InMaxUsedID), IDAcquireTag(InIDAcquireTag)
@@ -181,6 +190,13 @@ struct FDataSetMeta
 		: InputRegisters(nullptr), NumVariables(0), DataSetSizeInBytes(0), DataSetAccessIndex(INDEX_NONE), DataSetOffset(0), InstanceOffset(0)
 		, IDTable(nullptr), FreeIDTable(nullptr), NumFreeIDs(nullptr), MaxUsedID(nullptr), IDAcquireTag(0)
 	{}
+
+private:
+	// Non-copyable and non-movable
+	FDataSetMeta(FDataSetMeta&&) = delete;
+	FDataSetMeta(const FDataSetMeta&) = delete;
+	FDataSetMeta& operator=(FDataSetMeta&&) = delete;
+	FDataSetMeta& operator=(const FDataSetMeta&) = delete;
 };
 
 namespace VectorVM
@@ -199,6 +215,22 @@ namespace VectorVM
 	};
 }
 
+//Data the VM will keep on each dataset locally per thread which is then thread safely pushed to it's destination at the end of execution.
+struct FDataSetThreadLocalTempData
+{
+	FDataSetThreadLocalTempData()
+		:MaxID(INDEX_NONE)
+	{}
+
+	TArray<int32> IDsToFree;
+	int32 MaxID;
+
+	//TODO: Possibly store output data locally and memcpy to the real buffers. Could avoid false sharing in parallel execution and so improve perf.
+	//using _mm_stream_ps on platforms that support could also work for this?
+	//TArray<TArray<float>> OutputFloatData;
+	//TArray<TArray<int32>> OutputIntData;
+};
+
 /**
 * Context information passed around during VM execution.
 */
@@ -222,7 +254,9 @@ struct FVectorVMContext : TThreadSingleton<FVectorVMContext>
 	int32 StartInstance;
 
 	/** Array of meta data on data sets. TODO: This struct should be removed and all features it contains be handled by more general vm ops and the compiler's knowledge of offsets etc. */
-	FDataSetMeta*RESTRICT DataSetMetaTable;
+	TArray<FDataSetMeta>* RESTRICT DataSetMetaTable;
+
+	TArray<FDataSetThreadLocalTempData> ThreadLocalTempData;
 
 #if STATS
 	TArray<FCycleCounter, TInlineAllocator<64>> StatCounterStack;
@@ -247,12 +281,14 @@ struct FVectorVMContext : TThreadSingleton<FVectorVMContext>
 		int32 InNumSecondaryDatasets,
 		FVMExternalFunction* InExternalFunctionTable,
 		void** InUserPtrTable,
-		FDataSetMeta*RESTRICT InDataSetMetaTable
+		TArray<FDataSetMeta>& RESTRICT InDataSetMetaTable
 #if STATS
 		, const TArray<TStatId>* InStatScopes
 #endif
 	);
 
+	void FinishExec();
+
 	void PrepareForChunk(const uint8* InCode, int32 InNumInstances, int32 InStartInstance)
 	{
 		Code = InCode;
