Fix persistent ID to index table writing incorrect index

diff --git a/Engine/Source/Runtime/VectorVM/Private/VectorVM.cpp b/Engine/Source/Runtime/VectorVM/Private/VectorVM.cpp
--- a/Engine/Source/Runtime/VectorVM/Private/VectorVM.cpp
+++ b/Engine/Source/Runtime/VectorVM/Private/VectorVM.cpp
@@ -1137,8 +1137,7 @@ struct FScalarKernelUpdateID
 			else
 			{
 				//Update the actual index for this ID. No thread safety is needed as this ID slot can only ever be written by this instance and so a single thread.
-				int32 RealIdx = InstanceOffset + Index;	
-				IDTable[InstanceId] = RealIdx;
+				IDTable[InstanceId] = Index;
 
 				//Update thread local max ID seen. We push this to the real value at the end of execution.
 				DataSetTempData.MaxID = FMath::Max(DataSetTempData.MaxID, InstanceId);
