Fixing Simon Tovey's previous checkin on non-Stat builds

diff --git a/Engine/Source/Runtime/VectorVM/Private/VectorVM.cpp b/Engine/Source/Runtime/VectorVM/Private/VectorVM.cpp
--- a/Engine/Source/Runtime/VectorVM/Private/VectorVM.cpp
+++ b/Engine/Source/Runtime/VectorVM/Private/VectorVM.cpp
@@ -789,11 +789,13 @@ struct FVectorKernelEnterStatScope
 	static VM_FORCEINLINE void Exec(FVectorVMContext& Context)
 	{
 		FConstantHandler<int32> ScopeIdx(Context);
-		if (STATS && GbDetailedVMScriptStats)
+#if STATS
+		if (GbDetailedVMScriptStats)
 		{
 			int32 CounterIdx = Context.StatCounterStack.AddDefaulted(1);
 			Context.StatCounterStack[CounterIdx].Start((*Context.StatScopes)[ScopeIdx.Get()]);
 		}
+#endif
 	}
 };
 
@@ -801,11 +803,14 @@ struct FVectorKernelExitStatScope
 {
 	static VM_FORCEINLINE void Exec(FVectorVMContext& Context)
 	{
-		if (STATS && GbDetailedVMScriptStats)
+
+#if STATS
+		if (GbDetailedVMScriptStats)
 		{
 			Context.StatCounterStack.Last().Stop();
 			Context.StatCounterStack.Pop(false);
 		}
+#endif
 	}
 };
 
