Added "vm.InstancesPerChunk" to replace VectorVM::InstancesPerChunk

diff --git a/Engine/Source/Runtime/VectorVM/Private/VectorVM.cpp b/Engine/Source/Runtime/VectorVM/Private/VectorVM.cpp
--- a/Engine/Source/Runtime/VectorVM/Private/VectorVM.cpp
+++ b/Engine/Source/Runtime/VectorVM/Private/VectorVM.cpp
@@ -90,6 +90,14 @@ static FAutoConsoleVariableRef CVarDetailedVMScriptStats(
 	ECVF_Default
 );
 
+static int32 GParallelVVMInstancesPerChunk = 128;
+static FAutoConsoleVariableRef CVarParallelVVMInstancesPerChunk(
+	TEXT("vm.InstancesPerChunk"),
+	GParallelVVMInstancesPerChunk,
+	TEXT("Number of instances per VM chunk. (default=128) \n"),
+	ECVF_ReadOnly
+);
+
 //////////////////////////////////////////////////////////////////////////
 //  Constant Handlers
 
@@ -203,7 +211,7 @@ FVectorVMContext::FVectorVMContext()
 	, StatScopes(nullptr)
 #endif
 {
-	uint32 TempRegisterSize = Align((VectorVM::InstancesPerChunk) * VectorVM::MaxInstanceSizeBytes, VECTOR_WIDTH_BYTES) + VECTOR_WIDTH_BYTES;
+	uint32 TempRegisterSize = Align((GParallelVVMInstancesPerChunk) * VectorVM::MaxInstanceSizeBytes, VECTOR_WIDTH_BYTES) + VECTOR_WIDTH_BYTES;
 	TempRegTable.SetNumUninitialized(TempRegisterSize * VectorVM::NumTempRegisters);
 	// Map temporary registers.
 	for (int32 i = 0; i < VectorVM::NumTempRegisters; ++i)
@@ -1834,7 +1842,7 @@ void VectorVM::Exec(
 		DataSetIndexTable.Add(DataSetMetaTable[Idx].DataSetAccessIndex);	// prime counter index table with the data set offset; will be incremented with every write for each instance
 	}
 
-	int32 NumChunks = (NumInstances / InstancesPerChunk) + 1;
+	int32 NumChunks = (NumInstances / GParallelVVMInstancesPerChunk) + 1;
 	int32 ChunksPerBatch = (GbParallelVVM != 0 && FApp::ShouldUseThreadingForPerformance()) ? GParallelVVMChunksPerBatch : NumChunks;
 	int32 NumBatches = FMath::DivideAndRoundUp(NumChunks, ChunksPerBatch);
 	bool bParallel = NumBatches > 1;
@@ -1853,13 +1861,13 @@ void VectorVM::Exec(
 
 		// Process one chunk at a time.
 		int32 ChunkIdx = BatchIdx * ChunksPerBatch;
-		int32 FirstInstance = ChunkIdx * InstancesPerChunk;
-		int32 FinalInstance = FMath::Min(NumInstances, FirstInstance + (ChunksPerBatch * InstancesPerChunk));
+		int32 FirstInstance = ChunkIdx * GParallelVVMInstancesPerChunk;
+		int32 FinalInstance = FMath::Min(NumInstances, FirstInstance + (ChunksPerBatch * GParallelVVMInstancesPerChunk));
 		int32 InstancesLeft = FinalInstance - FirstInstance;
 		while (InstancesLeft > 0)
 		{
-			int32 NumInstancesThisChunk = FMath::Min(InstancesLeft, (int32)InstancesPerChunk);
-			int32 StartInstance = InstancesPerChunk * ChunkIdx;
+			int32 NumInstancesThisChunk = FMath::Min(InstancesLeft, (int32)GParallelVVMInstancesPerChunk);
+			int32 StartInstance = GParallelVVMInstancesPerChunk * ChunkIdx;
 			// Setup execution context.
 			Context.PrepareForChunk(Code, NumInstancesThisChunk, StartInstance);
 
@@ -1991,7 +1999,7 @@ void VectorVM::Exec(
 				}
 			} while (Op != EVectorVMOp::done);
 
-			InstancesLeft -= InstancesPerChunk;
+			InstancesLeft -= GParallelVVMInstancesPerChunk;
 			++ChunkIdx;
 		}
 
diff --git a/Engine/Source/Runtime/VectorVM/Private/VectorVMPrivate.h b/Engine/Source/Runtime/VectorVM/Private/VectorVMPrivate.h
--- a/Engine/Source/Runtime/VectorVM/Private/VectorVMPrivate.h
+++ b/Engine/Source/Runtime/VectorVM/Private/VectorVMPrivate.h
@@ -14,7 +14,6 @@ namespace VectorVM
 	/** Constants. */
 	enum
 	{
-		InstancesPerChunk = 128,
 		MaxInstanceSizeBytes = 4,
 	};
 }
