Made VM context TLS again as the pooled version was causing a perf issue.

diff --git a/Engine/Source/Runtime/VectorVM/Private/VectorVM.cpp b/Engine/Source/Runtime/VectorVM/Private/VectorVM.cpp
--- a/Engine/Source/Runtime/VectorVM/Private/VectorVM.cpp
+++ b/Engine/Source/Runtime/VectorVM/Private/VectorVM.cpp
@@ -197,8 +197,6 @@ public:
 
 //////////////////////////////////////////////////////////////////////////
 
-TLockFreePointerListLIFO<FVectorVMContext> FVectorVMContext::VMContextPool;
-
 FVectorVMContext::FVectorVMContext()
 	: Code(nullptr)
 	, ConstantTable(nullptr)
@@ -250,8 +248,9 @@ void FVectorVMContext::PrepareForExec(
 	StatCounterStack.Reserve(StatScopes->Num());
 #endif
 
-	uint32 TempRegisterSize = Align(MaxNumInstances * VectorVM::MaxInstanceSizeBytes, VECTOR_WIDTH_BYTES) + VECTOR_WIDTH_BYTES;
-	TempRegTable.SetNumUninitialized(TempRegisterSize * NumTempRegisters, false);
+	int32 TempRegisterSize = Align(MaxNumInstances * VectorVM::MaxInstanceSizeBytes, PLATFORM_CACHE_LINE_SIZE);
+	int32 TempBufferSize = TempRegisterSize * NumTempRegisters;
+	TempRegTable.SetNumUninitialized(TempBufferSize, false);
 	// Attempt to map temp registers more tightly packed for low instance counts to reduce cache misses.
 	for (int32 i = 0; i < NumTempRegisters; ++i)
 	{
@@ -1835,6 +1834,7 @@ void VectorVM::Exec(
 #endif
 	)
 {
+	//TRACE_CPUPROFILER_EVENT_SCOPE("VMExec");
 	SCOPE_CYCLE_COUNTER(STAT_VVMExec);
 
 	// table of index counters, one for each data set
@@ -1860,9 +1860,7 @@ void VectorVM::Exec(
 	{
 		//SCOPE_CYCLE_COUNTER(STAT_VVMExecChunk);
 
-		FVectorVMContext* ContextPtr = FVectorVMContext::GetContext();
-		checkSlow(ContextPtr);
-		FVectorVMContext& Context = *ContextPtr;
+		FVectorVMContext& Context = FVectorVMContext::Get();
 		Context.PrepareForExec(InputRegisters, OutputRegisters, NumTempRegisters, NumInputRegisters, NumOutputRegisters, ConstantTable, DataSetIndexTable.GetData(), DataSetOffsetTable.GetData(), DataSetOffsetTable.Num(),
 			ExternalFunctionTable, UserPtrTable, DataSetMetaTable, MaxInstances
 #if STATS
@@ -2015,7 +2013,6 @@ void VectorVM::Exec(
 		}
 
 		Context.FinishExec();
-		FVectorVMContext::ReleaseContext(ContextPtr);
 	};
 
 	ParallelFor(NumBatches, ExecChunkBatch, GbParallelVVM == 0 || !bParallel);
diff --git a/Engine/Source/Runtime/VectorVM/Public/VectorVM.h b/Engine/Source/Runtime/VectorVM/Public/VectorVM.h
--- a/Engine/Source/Runtime/VectorVM/Public/VectorVM.h
+++ b/Engine/Source/Runtime/VectorVM/Public/VectorVM.h
@@ -233,7 +233,7 @@ struct FDataSetThreadLocalTempData
 /**
 * Context information passed around during VM execution.
 */
-struct FVectorVMContext
+struct FVectorVMContext : TThreadSingleton<FVectorVMContext>
 {
 	/** Pointer to the next element in the byte code. */
 	uint8 const* RESTRICT Code;
@@ -309,30 +309,8 @@ struct FVectorVMContext
 		RandCounters.Reset();
 		RandCounters.SetNumZeroed(InNumInstances);
 	}
-
-	static TLockFreePointerListLIFO<FVectorVMContext> VMContextPool;
-	FORCEINLINE static FVectorVMContext* GetContext();
-	FORCEINLINE static void ReleaseContext(FVectorVMContext* Context);
 };
 
-FORCEINLINE FVectorVMContext* FVectorVMContext::GetContext()
-{
-	if (FVectorVMContext* Ret = VMContextPool.Pop())
-	{
-		return Ret;
-	}
-	else
-	{
-		return new FVectorVMContext();
-	}
-}
-
-FORCEINLINE void FVectorVMContext::ReleaseContext(FVectorVMContext* Context)
-{
-	checkSlow(Context);
-	VMContextPool.Push(Context);
-}
-
 namespace VectorVM
 {
 	/** Get total number of op-codes */
