- Removing a bunch of setup work we do for VM exectuion with remapping of register tables.
- VM is now directly fed a set of pre generated register tables from the datasets.
- Split the monolithic register table in the VM up so there are explicit I/O and temp register tables the script indexes into directly.
- Avoids some recreation of expensve objects in favour of manual reset calls.
- Re-wrote Oupt kernel to be more explicit. Going via templated handler in this case didn't get us any code reuse and just obfuscated it's workings.

Saves ~10-25us of overhead per VM involcation which soon adds up.
Saves ~1-2us inside each VM exec itself.

diff --git a/Engine/Plugins/FX/Niagara/Source/Niagara/Classes/NiagaraDataSet.h b/Engine/Plugins/FX/Niagara/Source/Niagara/Classes/NiagaraDataSet.h
--- a/Engine/Plugins/FX/Niagara/Source/Niagara/Classes/NiagaraDataSet.h
+++ b/Engine/Plugins/FX/Niagara/Source/Niagara/Classes/NiagaraDataSet.h
@@ -9,8 +9,6 @@
 #include "RenderingThread.h"
 #include "NiagaraDataSet.generated.h"
 
-typedef TArray<uint8*, TInlineAllocator<128>> FNiagaraRegisterTable;
-
 /** Helper class defining the layout and location of an FNiagaraVariable in an FNiagaraDataBuffer-> */
 USTRUCT()
 struct FNiagaraVariableLayoutInfo
@@ -113,11 +111,10 @@ public:
 	void GPUCopyFrom(float* GPUReadBackFloat, int* GPUReadBackInt, int32 StartIdx, int32 NumInstances, uint32 InSrcFloatStride, uint32 InSrcIntStride);
 	void Dump(int32 StartIndex, int32 NumInstances, const FString& Label)const;
 
-	FORCEINLINE_DEBUGGABLE void AppendToRegisterTable(FNiagaraRegisterTable& Registers, int32 StartInstance);
-	FORCEINLINE void ClearRegisterTable(FNiagaraRegisterTable& Registers);
+	FORCEINLINE TArray<uint8*>& GetRegisterTable() { return RegisterTable; }
 
-	const TArray<uint8>& GetFloatBuffer()const { return FloatData; }
-	const TArray<uint8>& GetInt32Buffer()const { return Int32Data; }
+	FORCEINLINE const TArray<uint8>& GetFloatBuffer()const { return FloatData; }
+	FORCEINLINE const TArray<uint8>& GetInt32Buffer()const { return Int32Data; }
 
 	FORCEINLINE const uint8* GetComponentPtrFloat(uint32 ComponentIdx)const	{ return FloatData.GetData() + FloatStride * ComponentIdx; }
 	FORCEINLINE const uint8* GetComponentPtrInt32(uint32 ComponentIdx)const	{ return Int32Data.GetData() + Int32Stride * ComponentIdx; }
@@ -1521,38 +1518,3 @@ FORCEINLINE void FNiagaraDataBuffer::CheckUsage(bool bReadOnly)const
 	checkSlow((Owner->GetSimTarget() == ENiagaraSimTarget::CPUSim  && (IsInGameThread() || (bReadOnly || !IsInRenderingThread()))) ||
 		(Owner->GetSimTarget() == ENiagaraSimTarget::GPUComputeSim) && IsInRenderingThread());
 }
-
-FORCEINLINE_DEBUGGABLE void FNiagaraDataBuffer::AppendToRegisterTable(FNiagaraRegisterTable& Registers, int32 StartInstance)
-{
-	checkSlow(Owner && Owner->IsInitialized());
-	checkSlow(Owner->GetSimTarget() == ENiagaraSimTarget::CPUSim);
-	CheckUsage(true);
-
-	if (!ensure(RegisterTable.Num() != 0))
-	{
-		BuildRegisterTable();
-	}
-
-	int32 RegisterCount = RegisterTable.Num();
-	int32 NumRegisters = Registers.Num();
-	Registers.AddUninitialized(RegisterCount);
-
-	if (StartInstance == 0)
-	{
-		FMemory::Memcpy(&Registers[NumRegisters], RegisterTable.GetData(), sizeof(uint8*) * RegisterCount);
-	}
-	else
-	{
-		int32 StartOffset = sizeof(float) * StartInstance;
-		for (int32 i = 0; i < RegisterCount; ++i)
-		{
-			Registers[NumRegisters + i] = RegisterTable[i] + StartOffset;
-		}
-	}
-}
-
-FORCEINLINE void FNiagaraDataBuffer::ClearRegisterTable(FNiagaraRegisterTable& Registers)
-{
-	int32 RegisterCount = RegisterTable.Num();
-	Registers.AddZeroed(RegisterCount);
-}
\ No newline at end of file
diff --git a/Engine/Plugins/FX/Niagara/Source/Niagara/Classes/NiagaraScriptExecutionContext.h b/Engine/Plugins/FX/Niagara/Source/Niagara/Classes/NiagaraScriptExecutionContext.h
--- a/Engine/Plugins/FX/Niagara/Source/Niagara/Classes/NiagaraScriptExecutionContext.h
+++ b/Engine/Plugins/FX/Niagara/Source/Niagara/Classes/NiagaraScriptExecutionContext.h
@@ -66,10 +66,11 @@ struct FNiagaraDataSetExecutionInfo
 		, StartInstance(0)
 		, bUpdateInstanceCount(false)
 	{
+		Reset();
 	}
 
 
-	void Init(FNiagaraDataSet* InDataSet, FNiagaraDataBuffer* InInput, FNiagaraDataBuffer* InOutput, int32 InStartInstance, bool bInUpdateInstanceCount)
+	FORCEINLINE void Init(FNiagaraDataSet* InDataSet, FNiagaraDataBuffer* InInput, FNiagaraDataBuffer* InOutput, int32 InStartInstance, bool bInUpdateInstanceCount)
 	{
 		if (Input)
 		{
@@ -95,11 +96,21 @@ struct FNiagaraDataSetExecutionInfo
 	
 	~FNiagaraDataSetExecutionInfo()
 	{
-		check(Output == nullptr || Output->IsBeingWritten());
+		Reset();
+	}
+
+	FORCEINLINE void Reset()
+	{
 		if (Input)
 		{
 			Input->ReleaseReadRef();
 		}
+
+		DataSet = nullptr;
+		Input = nullptr;
+		Output = nullptr;
+		StartInstance = INDEX_NONE;
+		bUpdateInstanceCount = false;
 	}
 
 	FNiagaraDataSet* DataSet;
@@ -122,7 +133,7 @@ struct FNiagaraScriptExecutionContext
 	/** Parameter store. Contains all data interfaces and a parameter buffer that can be used directly by the VM or GPU. */
 	FNiagaraScriptExecutionParameterStore Parameters;
 
-	TArray<FDataSetMeta> DataSetMetaTable;
+	TArray<FDataSetMeta, TInlineAllocator<4>> DataSetMetaTable;
 
 	TArray<FNiagaraDataSetExecutionInfo, TInlineAllocator<4>> DataSetInfo;
 
@@ -137,7 +148,7 @@ struct FNiagaraScriptExecutionContext
 	void PostTick();
 
 	void BindData(int32 Index, FNiagaraDataSet& DataSet, int32 StartInstance, bool bUpdateInstanceCounts);
-	void BindData(int32 Index, FNiagaraDataBuffer* Input, FNiagaraDataBuffer* Output, int32 StartInstance, bool bUpdateInstanceCounts);
+	void BindData(int32 Index, FNiagaraDataBuffer* Input, int32 StartInstance, bool bUpdateInstanceCounts);
 	bool Execute(uint32 NumInstances);
 
 	const TArray<UNiagaraDataInterface*>& GetDataInterfaces()const { return Parameters.GetDataInterfaces(); }
diff --git a/Engine/Plugins/FX/Niagara/Source/Niagara/Private/NiagaraEmitterInstance.cpp b/Engine/Plugins/FX/Niagara/Source/Niagara/Private/NiagaraEmitterInstance.cpp
--- a/Engine/Plugins/FX/Niagara/Source/Niagara/Private/NiagaraEmitterInstance.cpp
+++ b/Engine/Plugins/FX/Niagara/Source/Niagara/Private/NiagaraEmitterInstance.cpp
@@ -1546,7 +1546,7 @@ void FNiagaraEmitterInstance::Tick(float DeltaSeconds)
 					EventExecCountBindings[EventScriptIdx].SetValue(EventNumToSpawn);
 
 					EventExecContexts[EventScriptIdx].BindData(0, Data, EventSpawnStart, false);
-					EventExecContexts[EventScriptIdx].BindData(1, EventHandlingInfo[EventScriptIdx].EventData, nullptr, i, false);
+					EventExecContexts[EventScriptIdx].BindData(1, EventHandlingInfo[EventScriptIdx].EventData, i, false);
 					EventExecContexts[EventScriptIdx].Execute(EventNumToSpawn);
 
 					if (GbDumpParticleData || System->bDumpDebugEmitterInfo)
@@ -1605,7 +1605,7 @@ void FNiagaraEmitterInstance::Tick(float DeltaSeconds)
 					EventExecCountBindings[EventScriptIdx].SetValue(NumInstancesPrev);
 
 					EventExecContexts[EventScriptIdx].BindData(0, Data, 0, true);
-					EventExecContexts[EventScriptIdx].BindData(1, EventData, nullptr, i, false);
+					EventExecContexts[EventScriptIdx].BindData(1, EventData, i, false);
 					EventExecContexts[EventScriptIdx].Execute(NumInstancesPrev);
 
 					Data.EndSimulate();
diff --git a/Engine/Plugins/FX/Niagara/Source/Niagara/Private/NiagaraScriptExecutionContext.cpp b/Engine/Plugins/FX/Niagara/Source/Niagara/Private/NiagaraScriptExecutionContext.cpp
--- a/Engine/Plugins/FX/Niagara/Source/Niagara/Private/NiagaraScriptExecutionContext.cpp
+++ b/Engine/Plugins/FX/Niagara/Source/Niagara/Private/NiagaraScriptExecutionContext.cpp
@@ -149,16 +149,29 @@ void FNiagaraScriptExecutionContext::PostTick()
 
 void FNiagaraScriptExecutionContext::BindData(int32 Index, FNiagaraDataSet& DataSet, int32 StartInstance, bool bUpdateInstanceCounts)
 {
+	FNiagaraDataBuffer* Input = DataSet.GetCurrentData();
+	FNiagaraDataBuffer* Output = DataSet.GetDestinationData();
+
 	DataSetInfo.SetNum(FMath::Max(DataSetInfo.Num(), Index + 1));
-	
-	DataSetInfo[Index].Init(&DataSet, DataSet.GetCurrentData(), DataSet.GetDestinationData(), StartInstance, bUpdateInstanceCounts);
+	DataSetInfo[Index].Init(&DataSet, Input, Output, StartInstance, bUpdateInstanceCounts);
+
+	//Would be nice to roll this and DataSetInfo into one but currently the VM being in it's own Engine module prevents this. Possibly should move the VM into Niagara itself.
+	uint8** InputRegisters = Input ? Input->GetRegisterTable().GetData() : nullptr;
+	uint8** OutputRegisters = Output ? Output->GetRegisterTable().GetData() : nullptr;
+	DataSetMetaTable.SetNum(FMath::Max(DataSetMetaTable.Num(), Index + 1));
+	DataSetMetaTable[Index].Init(InputRegisters, OutputRegisters, StartInstance,
+		Output ? &Output->GetIDTable() : nullptr, &DataSet.GetFreeIDTable(), &DataSet.GetNumFreeIDs(), &DataSet.GetMaxUsedID(), DataSet.GetIDAcquireTag());
 }
 
-void FNiagaraScriptExecutionContext::BindData(int32 Index, FNiagaraDataBuffer* Input, FNiagaraDataBuffer* Output, int32 StartInstance, bool bUpdateInstanceCounts)
+void FNiagaraScriptExecutionContext::BindData(int32 Index, FNiagaraDataBuffer* Input, int32 StartInstance, bool bUpdateInstanceCounts)
 {
+	check(Input && Input->GetOwner());
 	DataSetInfo.SetNum(FMath::Max(DataSetInfo.Num(), Index + 1));
-	check(Input || Output);
-	DataSetInfo[Index].Init(Input ? Input->GetOwner() : Output->GetOwner(), Input, Output, StartInstance, bUpdateInstanceCounts);
+	FNiagaraDataSet* DataSet = Input->GetOwner();
+	DataSetInfo[Index].Init(DataSet, Input, nullptr, StartInstance, bUpdateInstanceCounts);
+
+	DataSetMetaTable.SetNum(FMath::Max(DataSetMetaTable.Num(), Index + 1));
+	DataSetMetaTable[Index].Init(Input->GetRegisterTable().GetData(), nullptr, StartInstance, nullptr, nullptr, &DataSet->GetNumFreeIDs(), &DataSet->GetMaxUsedID(), DataSet->GetIDAcquireTag());
 }
 
 bool FNiagaraScriptExecutionContext::Execute(uint32 NumInstances)
@@ -171,56 +184,12 @@ bool FNiagaraScriptExecutionContext::Execute(uint32 NumInstances)
 
 	++TickCounter;//Should this be per execution?
 
-	FNiagaraRegisterTable InputRegisters;
-	FNiagaraRegisterTable OutputRegisters;
-
-	DataSetMetaTable.Reset();
-
-	bool bError = false;
-	{
-		SCOPE_CYCLE_COUNTER(STAT_NiagaraSimRegisterSetup);
-		for (FNiagaraDataSetExecutionInfo& Info : DataSetInfo)
-		{
-			FNiagaraDataBuffer* DestinationData = Info.DataSet->GetDestinationData();
-#if NIAGARA_NAN_CHECKING
-			DataSetInfo.DataSet->CheckForNaNs();
-#endif
-			check(Info.DataSet);
-			int32 NumInputRegisters = InputRegisters.Num();
-			DataSetMetaTable.Emplace(NumInputRegisters, Info.StartInstance,
-				DestinationData ? &DestinationData->GetIDTable() : nullptr, &Info.DataSet->GetFreeIDTable(), &Info.DataSet->GetNumFreeIDs(), &Info.DataSet->GetMaxUsedID(), Info.DataSet->GetIDAcquireTag());
-
-			int32 TotalComponents = Info.DataSet->GetNumFloatComponents() + Info.DataSet->GetNumInt32Components();
-			if (Info.Input && Info.Input->GetNumInstances() > 0)
-			{
-				Info.Input->AppendToRegisterTable(InputRegisters, Info.StartInstance);
-			}
-			else
-			{
-				Info.DataSet->GetCurrentDataChecked().ClearRegisterTable(InputRegisters);
-			}
-
-			if (Info.Output)
-			{
-				Info.Output->AppendToRegisterTable(OutputRegisters, Info.StartInstance);
-			}
-			else
-			{
-				Info.DataSet->GetCurrentDataChecked().ClearRegisterTable(OutputRegisters);
-			}
-		}
-	}
-
-	if (GbExecVMScripts != 0 && !bError)
+	if (GbExecVMScripts != 0)
 	{
 		const FNiagaraVMExecutableData& ExecData = Script->GetVMExecutableData();
 		VectorVM::Exec(
 			ExecData.ByteCode.GetData(),
 			ExecData.NumTempRegisters,
-			InputRegisters.GetData(),
-			InputRegisters.Num(),
-			OutputRegisters.GetData(),
-			OutputRegisters.Num(),
 			Parameters.GetParameterDataArray().GetData(),
 			DataSetMetaTable,
 			FunctionTable.GetData(),
@@ -247,7 +216,12 @@ bool FNiagaraScriptExecutionContext::Execute(uint32 NumInstances)
 		}
 	}
 
-	DataSetInfo.Reset();
+	//Can maybe do without resetting here. Just doing it for tidiness.
+	for (int32 DataSetIdx = 0; DataSetIdx < DataSetInfo.Num(); ++DataSetIdx)
+	{
+		DataSetInfo[DataSetIdx].Reset();
+		DataSetMetaTable[DataSetIdx].Reset();
+	}
 
 	return true;//TODO: Error cases?
 }
@@ -574,4 +548,4 @@ void FNiagaraComputeExecutionContext::SetDataToRender(FNiagaraDataBuffer* InData
 	{
 		DataToRender->AddReadRef();
 	}
-}
\ No newline at end of file
+}
diff --git a/Engine/Plugins/FX/Niagara/Source/NiagaraCore/Private/NiagaraCustomVersion.cpp b/Engine/Plugins/FX/Niagara/Source/NiagaraCore/Private/NiagaraCustomVersion.cpp
--- a/Engine/Plugins/FX/Niagara/Source/NiagaraCore/Private/NiagaraCustomVersion.cpp
+++ b/Engine/Plugins/FX/Niagara/Source/NiagaraCore/Private/NiagaraCustomVersion.cpp
@@ -8,4 +8,4 @@ const FGuid FNiagaraCustomVersion::GUID(0xFCF57AFA, 0x50764283, 0xB9A9E658, 0xFF
 // Register the custom version with core
 FCustomVersionRegistration GRegisterNiagaraCustomVersion(FNiagaraCustomVersion::GUID, FNiagaraCustomVersion::LatestVersion, TEXT("NiagaraVer"));
 
-const FGuid FNiagaraCustomVersion::LatestScriptCompileVersion(0x908C3A78, 0x6E39F14C, 0x8E1C1A93, 0xDD667695);
\ No newline at end of file
+const FGuid FNiagaraCustomVersion::LatestScriptCompileVersion(0x1E9267F8, 0xEA19004A, 0xBF0BA00E, 0x9231686F);
\ No newline at end of file
diff --git a/Engine/Source/Developer/ShaderFormatVectorVM/Private/ir_vm_gen_bytecode_visitor.cpp b/Engine/Source/Developer/ShaderFormatVectorVM/Private/ir_vm_gen_bytecode_visitor.cpp
--- a/Engine/Source/Developer/ShaderFormatVectorVM/Private/ir_vm_gen_bytecode_visitor.cpp
+++ b/Engine/Source/Developer/ShaderFormatVectorVM/Private/ir_vm_gen_bytecode_visitor.cpp
@@ -574,7 +574,7 @@ struct op_base
 
 	virtual FString to_string() = 0;
 	virtual void add_to_bytecode(TArray<uint8>& bytecode) = 0;
-	virtual bool finalize(_mesa_glsl_parse_state* parse_state, FVectorVMCompilationOutput&CompilationOutput, const unsigned num_temp_registers, unsigned *registers, unsigned op_idx) { return true; }
+	virtual bool finalize(_mesa_glsl_parse_state* parse_state, FVectorVMCompilationOutput&CompilationOutput, TArray<uint32>& temp_registers, unsigned op_idx) { return true; }
 	virtual void validate(_mesa_glsl_parse_state* parse_state, unsigned op_idx) = 0;
 
 	virtual struct op_standard* as_standard() { return nullptr; }
@@ -597,7 +597,7 @@ struct op_base
 
 	}
 
-	static bool finalize_component_temporary_allocation(_mesa_glsl_parse_state* parse_state, FVectorVMCompilationOutput& CompilationOutput, variable_info_node* component, const unsigned num_temp_registers, unsigned *registers, unsigned op_idx)
+	static bool finalize_component_temporary_allocation(_mesa_glsl_parse_state* parse_state, FVectorVMCompilationOutput& CompilationOutput, variable_info_node* component, TArray<uint32>& temp_registers, unsigned op_idx)
 	{
 		check(component);
 
@@ -614,12 +614,12 @@ struct op_base
 			{
 				if (component->offset == INDEX_NONE)
 				{
-					for (unsigned j = 0; j < num_temp_registers; ++j)
+					for (int32 j = 0; j < temp_registers.Num(); ++j)
 					{
-						if (registers[j] < op_idx || registers[j] == INDEX_NONE)
+						if (temp_registers[j] < op_idx || temp_registers[j] == INDEX_NONE)
 						{
-							component->offset = VectorVM::FirstTempRegister + j;
-							registers[j] = component->last_read;
+							component->offset = j;
+							temp_registers[j] = component->last_read;
 							CompilationOutput.MaxTempRegistersUsed = FMath::Max((int32)j, CompilationOutput.MaxTempRegistersUsed);
 #if VM_VERBOSE_LOGGING
 							UE_LOG(LogVVMBackend, Log, TEXT("OP:%d | Comp:%p allocated Reg: %d | Last Read: %d |"), op_idx, component, j, component->last_read);
@@ -628,11 +628,18 @@ struct op_base
 						}
 					}
 
-					if (component->offset == INDEX_NONE)
-					{
-						FString ErrorStr = FString::Printf(TEXT("Failed to allocate temporary register for component %s of variable %s."), ANSI_TO_TCHAR(component->name), ANSI_TO_TCHAR(component->owner->root->name));
-						_mesa_glsl_error(parse_state, TCHAR_TO_ANSI(*ErrorStr));
-					}
+					component->offset = temp_registers.Num();
+					CompilationOutput.MaxTempRegistersUsed = FMath::Max(temp_registers.Num(), CompilationOutput.MaxTempRegistersUsed);
+					temp_registers.Add(component->last_read);
+#if VM_VERBOSE_LOGGING
+					UE_LOG(LogVVMBackend, Log, TEXT("OP:%d | Comp:%p allocated Reg: %d | Last Read: %d |"), op_idx, component, j, component->last_read);
+#endif
+ 
+// 					if (component->offset == INDEX_NONE)
+// 					{
+// 						FString ErrorStr = FString::Printf(TEXT("Failed to allocate temporary register for component %s of variable %s."), ANSI_TO_TCHAR(component->name), ANSI_TO_TCHAR(component->owner->root->name));
+// 						_mesa_glsl_error(parse_state, TCHAR_TO_ANSI(*ErrorStr));
+// 					}
 				}
 			}
 			else
@@ -775,9 +782,9 @@ struct op_standard : public op_base
 		emit_short(dest_component->offset, bytecode);
 	}
 
-	virtual bool finalize(_mesa_glsl_parse_state* parse_state, FVectorVMCompilationOutput&CompilationOutput, const unsigned num_temp_registers, unsigned *registers, unsigned op_idx)
+	virtual bool finalize(_mesa_glsl_parse_state* parse_state, FVectorVMCompilationOutput&CompilationOutput, TArray<uint32>& temp_registers, unsigned op_idx)
 	{
-		return finalize_component_temporary_allocation(parse_state, CompilationOutput, dest_component, num_temp_registers, registers, op_idx);
+		return finalize_component_temporary_allocation(parse_state, CompilationOutput, dest_component, temp_registers, op_idx);
 	}
 
 	virtual void validate(_mesa_glsl_parse_state* parse_state, unsigned op_idx)override
@@ -828,15 +835,15 @@ struct op_input : public op_base
 		check(op_code != EVectorVMOp::done);
 		emit_byte((uint8)op_code, bytecode);
 		emit_short(dataset_idx, bytecode);
-		emit_short(register_idx + VectorVM::FirstInputRegister, bytecode);
+		emit_short(register_idx, bytecode);
 		//Destination is a temp register which we emit as usual.
 		check(dest_component);
 		emit_short(dest_component->offset, bytecode);
 	}
 
-	virtual bool finalize(_mesa_glsl_parse_state* parse_state, FVectorVMCompilationOutput&CompilationOutput, const unsigned num_temp_registers, unsigned *registers, unsigned op_idx)
+	virtual bool finalize(_mesa_glsl_parse_state* parse_state, FVectorVMCompilationOutput&CompilationOutput, TArray<uint32>& temp_registers, unsigned op_idx)
 	{
-		return finalize_component_temporary_allocation(parse_state, CompilationOutput, dest_component, num_temp_registers, registers, op_idx);
+		return finalize_component_temporary_allocation(parse_state, CompilationOutput, dest_component, temp_registers, op_idx);
 	}
 
 	virtual void validate(_mesa_glsl_parse_state* parse_state, unsigned op_idx)override
@@ -876,15 +883,15 @@ struct op_input_noadvance : public op_input
 		check(op_code != EVectorVMOp::done);
 		emit_byte((uint8)op_code, bytecode);
 		emit_short(dataset_idx, bytecode);
-		emit_short(register_idx + VectorVM::FirstInputRegister, bytecode);
+		emit_short(register_idx, bytecode);
 		//Destination is a temp register which we emit as usual.
 		check(dest_component);
 		emit_short(dest_component->offset, bytecode);
 	}
 
-	virtual bool finalize(_mesa_glsl_parse_state* parse_state, FVectorVMCompilationOutput&CompilationOutput, const unsigned num_temp_registers, unsigned *registers, unsigned op_idx)
+	virtual bool finalize(_mesa_glsl_parse_state* parse_state, FVectorVMCompilationOutput&CompilationOutput, TArray<uint32>& temp_registers, unsigned op_idx)
 	{
-		return finalize_component_temporary_allocation(parse_state, CompilationOutput, dest_component, num_temp_registers, registers, op_idx);
+		return finalize_component_temporary_allocation(parse_state, CompilationOutput, dest_component, temp_registers, op_idx);
 	}
 
 	virtual void validate(_mesa_glsl_parse_state* parse_state, unsigned op_idx)override
@@ -931,7 +938,7 @@ struct op_output : public op_base
 		emit_short(dataset_index, bytecode);
 		emit_short(instance_idx_component->offset, bytecode);
 		emit_short(value_component->offset, bytecode);
-		emit_short(register_idx + VectorVM::FirstOutputRegister, bytecode);
+		emit_short(register_idx, bytecode);
 	}
 
 	virtual void validate(_mesa_glsl_parse_state* parse_state, unsigned op_idx)override
@@ -977,9 +984,9 @@ struct op_index_acquire : public op_base
 		emit_short(dest_component->offset, bytecode);
 	}
 
-	virtual bool finalize(_mesa_glsl_parse_state* parse_state, FVectorVMCompilationOutput&CompilationOutput, const unsigned num_temp_registers, unsigned *registers, unsigned op_idx)
+	virtual bool finalize(_mesa_glsl_parse_state* parse_state, FVectorVMCompilationOutput&CompilationOutput, TArray<uint32>& temp_registers, unsigned op_idx)
 	{
-		return finalize_component_temporary_allocation(parse_state, CompilationOutput, dest_component, num_temp_registers, registers, op_idx);
+		return finalize_component_temporary_allocation(parse_state, CompilationOutput, dest_component, temp_registers, op_idx);
 	}
 
 	virtual void validate(_mesa_glsl_parse_state* parse_state, unsigned op_idx)override
@@ -1021,14 +1028,14 @@ struct op_id_acquire : public op_base
 		emit_short(id_tag_component->offset, bytecode);
 	}
 
-	virtual bool finalize(_mesa_glsl_parse_state* parse_state, FVectorVMCompilationOutput&CompilationOutput, const unsigned num_temp_registers, unsigned *registers, unsigned op_idx)
+	virtual bool finalize(_mesa_glsl_parse_state* parse_state, FVectorVMCompilationOutput&CompilationOutput, TArray<uint32>& temp_registers, unsigned op_idx)
 	{
-		if (!finalize_component_temporary_allocation(parse_state, CompilationOutput, id_index_component, num_temp_registers, registers, op_idx))
+		if (!finalize_component_temporary_allocation(parse_state, CompilationOutput, id_index_component, temp_registers, op_idx))
 		{
 			_mesa_glsl_error(parse_state, "Failed to allocate temporary registers");
 			return false;
 		}
-		if (!finalize_component_temporary_allocation(parse_state, CompilationOutput, id_tag_component, num_temp_registers, registers, op_idx))
+		if (!finalize_component_temporary_allocation(parse_state, CompilationOutput, id_tag_component, temp_registers, op_idx))
 		{
 			_mesa_glsl_error(parse_state, "Failed to allocate temporary registers");
 			return false;
@@ -1079,7 +1086,7 @@ struct op_id_update : public op_base
 		emit_short(index_component->offset, bytecode);
 	}
 
-	virtual bool finalize(_mesa_glsl_parse_state* parse_state, FVectorVMCompilationOutput&CompilationOutput, const unsigned num_temp_registers, unsigned *registers, unsigned op_idx)
+	virtual bool finalize(_mesa_glsl_parse_state* parse_state, FVectorVMCompilationOutput&CompilationOutput, TArray<uint32>& temp_registers, unsigned op_idx)
 	{
 		return true;
 	}
@@ -1190,8 +1197,7 @@ struct op_external_func : public op_base
 	{
 		_mesa_glsl_parse_state* parse_state;
 		FVectorVMCompilationOutput* CompilationOutput;
-		unsigned num_temp_registers;
-		unsigned *registers;
+		TArray<uint32>* temp_registers;
 		unsigned op_idx;
 		bool success;
 	};
@@ -1202,10 +1208,10 @@ struct op_external_func : public op_base
 		check(node->is_scalar());
 
 		//Some outputs for an external function can have invalid locations. If the thing has *any* valid outputs then it's considered a success.
-		ctx->success |= finalize_component_temporary_allocation(ctx->parse_state, *ctx->CompilationOutput, node, ctx->num_temp_registers, ctx->registers, ctx->op_idx);
+		ctx->success |= finalize_component_temporary_allocation(ctx->parse_state, *ctx->CompilationOutput, node, *ctx->temp_registers, ctx->op_idx);
 	}
 
-	virtual bool finalize(_mesa_glsl_parse_state* parse_state, FVectorVMCompilationOutput&CompilationOutput, const unsigned num_temp_registers, unsigned *registers, unsigned op_idx)
+	virtual bool finalize(_mesa_glsl_parse_state* parse_state, FVectorVMCompilationOutput&CompilationOutput, TArray<uint32>& temp_registers, unsigned op_idx)
 	{
 		//don't allocate any temp registers if we're culling the function.
 		function_index = -1;
@@ -1228,8 +1234,7 @@ struct op_external_func : public op_base
 		finalize_temp_register_ctx ctx;
 		ctx.parse_state = parse_state;
 		ctx.CompilationOutput = &CompilationOutput;
-		ctx.num_temp_registers = num_temp_registers;
-		ctx.registers = registers;
+		ctx.temp_registers = &temp_registers;
 		ctx.op_idx = op_idx;
 		ctx.success = true;
 		for (variable_info_node* output : outputs)
@@ -1332,7 +1337,7 @@ class ir_gen_vvm_visitor : public ir_hierarchical_visitor
 	{
 		ir_gen_vvm_visitor* v = (ir_gen_vvm_visitor*)user_ptr;
 		check(node->is_scalar());
-		node->offset = VectorVM::FirstInputRegister + v->num_input_components++;
+		node->offset = v->num_input_components++;
 	}
 
 	/** assigns an output register to a variable_info_node. */
@@ -1340,7 +1345,7 @@ class ir_gen_vvm_visitor : public ir_hierarchical_visitor
 	{
 		ir_gen_vvm_visitor* v = (ir_gen_vvm_visitor*)user_ptr;
 		check(node->is_scalar());
-		node->offset = VectorVM::FirstOutputRegister + v->num_output_components++;
+		node->offset = v->num_output_components++;
 	}
 
 	/** assigns an uniform buffer address to a variable_info_node. */
@@ -2070,9 +2075,11 @@ class ir_gen_vvm_visitor : public ir_hierarchical_visitor
 	void Finalize()
 	{
 		//Allocate temp registers
-		const unsigned num_temp_registers = VectorVM::NumTempRegisters;
-		unsigned *registers = rzalloc_array(mem_ctx, unsigned, num_temp_registers);
-		FMemory::Memset(registers, 0xFF, sizeof(unsigned) * num_temp_registers);
+// 		const unsigned num_temp_registers = VectorVM::NumTempRegisters;
+// 		unsigned *registers = rzalloc_array(mem_ctx, unsigned, num_temp_registers);
+// 		FMemory::Memset(registers, 0xFF, sizeof(unsigned) * num_temp_registers);
+		TArray<uint32> temp_registers;
+		temp_registers.Reserve(100);
 #if VM_VERBOSE_LOGGING
 		UE_LOG(LogVVMBackend, Log, TEXT("\n-------------------------------\nTemporary Allocations\n------------------------------\n"));
 #endif
@@ -2080,7 +2087,7 @@ class ir_gen_vvm_visitor : public ir_hierarchical_visitor
 		stripped_ops.Reserve(ordered_ops.Num());
 		for (int32 i = 0; i < ordered_ops.Num(); ++i)
 		{
-			if (ordered_ops[i]->finalize(parse_state, CompilationOutput, num_temp_registers, registers, i))
+			if (ordered_ops[i]->finalize(parse_state, CompilationOutput, temp_registers, i))
 			{
 				stripped_ops.Add(ordered_ops[i]);
 				ordered_ops[i]->validate(parse_state, i);
diff --git a/Engine/Source/Runtime/VectorVM/Private/VectorVM.cpp b/Engine/Source/Runtime/VectorVM/Private/VectorVM.cpp
--- a/Engine/Source/Runtime/VectorVM/Private/VectorVM.cpp
+++ b/Engine/Source/Runtime/VectorVM/Private/VectorVM.cpp
@@ -123,17 +123,6 @@ struct FConstantHandler : public FConstantHandlerBase
 	FORCEINLINE const T& GetAndAdvance() { return Constant; }
 };
 
-struct FDataSetOffsetHandler : FConstantHandlerBase
-{
-	uint32 Offset;
-	FDataSetOffsetHandler(FVectorVMContext& Context)
-		: FConstantHandlerBase(Context)
-		, Offset(Context.DataSetOffsetTable[ConstantIndex])
-	{}
-	FORCEINLINE const uint32 Get() { return Offset; }
-	FORCEINLINE const uint32 GetAndAdvance() { return Offset; }
-};
-
 template<>
 struct FConstantHandler<VectorRegister> : public FConstantHandlerBase
 {
@@ -180,7 +169,7 @@ private:
 public:
 	FORCEINLINE FRegisterHandler(FVectorVMContext& Context)
 		: FRegisterHandlerBase(Context)
-		, Register((T*)Context.RegisterTable[RegisterIndex])
+		, Register((T*)Context.GetTempRegister(RegisterIndex))
 	{}
 	FORCEINLINE const T Get() { return *Register; }
 	FORCEINLINE T* GetDest() { return Register; }
@@ -200,9 +189,6 @@ public:
 FVectorVMContext::FVectorVMContext()
 	: Code(nullptr)
 	, ConstantTable(nullptr)
-	, DataSetIndexTable(nullptr)
-	, DataSetOffsetTable(nullptr)
-	, NumSecondaryDataSets(0)
 	, ExternalFunctionTable(nullptr)
 	, UserPtrTable(nullptr)
 	, NumInstances(0)
@@ -210,23 +196,18 @@ FVectorVMContext::FVectorVMContext()
 #if STATS
 	, StatScopes(nullptr)
 #endif
+	, TempRegisterSize(0)
+	, TempBufferSize(0)
 {
 	RandStream.GenerateNewSeed();
 }
 
 void FVectorVMContext::PrepareForExec(
-	uint8*RESTRICT*RESTRICT InputRegisters,
-	uint8*RESTRICT*RESTRICT OutputRegisters,
 	int32 InNumTempRegisters,
-	int32 InNumInputRegisters,
-	int32 InNumOutputRegisters,
 	const uint8* InConstantTable,
-	int32 *InDataSetIndexTable,
-	int32 *InDataSetOffsetTable,
-	int32 InNumSecondaryDatasets,
 	FVMExternalFunction* InExternalFunctionTable,
 	void** InUserPtrTable,
-	TArray<FDataSetMeta>& RESTRICT InDataSetMetaTable,
+	TArrayView<FDataSetMeta> InDataSetMetaTable,
 	int32 MaxNumInstances
 #if STATS
 	, const TArray<TStatId>* InStatScopes
@@ -234,12 +215,7 @@ void FVectorVMContext::PrepareForExec(
 )
 {
 	NumTempRegisters = InNumTempRegisters;
-	NumInputRegisters = InNumInputRegisters;
-	NumOutputRegisters = InNumOutputRegisters;
 	ConstantTable = InConstantTable;
-	DataSetIndexTable = InDataSetIndexTable;
-	DataSetOffsetTable = InDataSetOffsetTable;
-	NumSecondaryDataSets = InNumSecondaryDatasets;
 	ExternalFunctionTable = InExternalFunctionTable;
 	UserPtrTable = InUserPtrTable;
 #if STATS
@@ -248,61 +224,48 @@ void FVectorVMContext::PrepareForExec(
 	StatCounterStack.Reserve(StatScopes->Num());
 #endif
 
-	int32 TempRegisterSize = Align(MaxNumInstances * VectorVM::MaxInstanceSizeBytes, PLATFORM_CACHE_LINE_SIZE);
-	int32 TempBufferSize = TempRegisterSize * NumTempRegisters;
+	TempRegisterSize = Align(MaxNumInstances * VectorVM::MaxInstanceSizeBytes, PLATFORM_CACHE_LINE_SIZE);
+	TempBufferSize = TempRegisterSize * NumTempRegisters;
 	TempRegTable.SetNumUninitialized(TempBufferSize, false);
-	// Attempt to map temp registers more tightly packed for low instance counts to reduce cache misses.
-	for (int32 i = 0; i < NumTempRegisters; ++i)
-	{
-		RegisterTable[i] = TempRegTable.GetData() + TempRegisterSize * i;
-	}
 
-	//Map IO Registers
-	for (int32 i = 0; i < NumInputRegisters; ++i)
-	{
-		RegisterTable[VectorVM::NumTempRegisters + i] = InputRegisters[i];
-	}
-	for (int32 i = 0; i < NumOutputRegisters; ++i)
+	DataSetMetaTable = InDataSetMetaTable;
+
+	for (auto& TLSTempData : ThreadLocalTempData)
 	{
-		RegisterTable[VectorVM::NumTempRegisters + VectorVM::MaxInputRegisters + i] = OutputRegisters[i];
+		TLSTempData.Reset();
 	}
-
-	DataSetMetaTable = &InDataSetMetaTable;
-
-	ThreadLocalTempData.Reset(DataSetMetaTable->Num());
-	ThreadLocalTempData.SetNum(DataSetMetaTable->Num());
+	ThreadLocalTempData.SetNum(DataSetMetaTable.Num());
 }
 
 void FVectorVMContext::FinishExec()
 {
 	//At the end of executing each chunk we can push any thread local temporary data out to the main storage with locks or atomics.
 
-	TArray<FDataSetMeta>& MetaTable = *DataSetMetaTable;
-	check(ThreadLocalTempData.Num() == MetaTable.Num());
-	for(int32 DataSetIndex=0; DataSetIndex < MetaTable.Num(); ++DataSetIndex)
+	check(ThreadLocalTempData.Num() == DataSetMetaTable.Num());
+	for(int32 DataSetIndex=0; DataSetIndex < DataSetMetaTable.Num(); ++DataSetIndex)
 	{
 		FDataSetThreadLocalTempData&RESTRICT Data = ThreadLocalTempData[DataSetIndex];
 
 		if (Data.IDsToFree.Num() > 0)
 		{
-			TArray<int32>&RESTRICT FreeIDTable = *MetaTable[DataSetIndex].FreeIDTable;
-			int32&RESTRICT NumFreeIDs = *MetaTable[DataSetIndex].NumFreeIDs;
+			TArray<int32>&RESTRICT FreeIDTable = *DataSetMetaTable[DataSetIndex].FreeIDTable;
+			int32&RESTRICT NumFreeIDs = *DataSetMetaTable[DataSetIndex].NumFreeIDs;
 			check(FreeIDTable.Num() >= NumFreeIDs + Data.IDsToFree.Num());
 
 			//Temporarily locking the free table until we can implement something lock-free
-			MetaTable[DataSetIndex].LockFreeTable();
+			DataSetMetaTable[DataSetIndex].LockFreeTable();
 			for (int32 IDToFree : Data.IDsToFree)
 			{
 				//UE_LOG(LogVectorVM, Warning, TEXT("AddFreeID: ID:%d | FreeTableIdx:%d."), IDToFree, NumFreeIDs);
 				FreeIDTable[NumFreeIDs++] = IDToFree;
 			}
 			//Unlock the free table.
-			MetaTable[DataSetIndex].UnlockFreeTable();
+			DataSetMetaTable[DataSetIndex].UnlockFreeTable();
 			Data.IDsToFree.Reset();
 		}
 
 		//Also update the max ID seen. This should be the ONLY place in the VM we update this max value.
-		volatile int32* MaxUsedID = MetaTable[DataSetIndex].MaxUsedID;
+		volatile int32* MaxUsedID = DataSetMetaTable[DataSetIndex].MaxUsedID;
 		int32 LocalMaxUsedID;
 		do
 		{
@@ -383,26 +346,6 @@ struct TTrinaryKernelHandler
 };
 
 
-template<typename Kernel, typename DstHandler, typename Arg0Handler, typename Arg1Handler, typename Arg2Handler, int32 NumInstancesPerOp>
-struct TTrinaryOutputKernelHandler
-{
-	static void Exec(FVectorVMContext& Context)
-	{
-		Arg0Handler Arg0(Context);
-		Arg1Handler Arg1(Context);
-		Arg2Handler Arg2(Context);
-
-		DstHandler Dst(Context, Arg0.Get());
-
-		int32 LoopInstances = Align(Context.NumInstances, NumInstancesPerOp) / NumInstancesPerOp;
-		for (int32 i = 0; i < LoopInstances; ++i)
-		{
-			Kernel::DoKernel(Context, Dst.GetDestAndAdvance(), Arg0.GetAndAdvance(), Arg1.GetAndAdvance(), Arg2.GetAndAdvance());
-		}
-	}
-};
-
-
 /** Base class of vector kernels with a single operand. */
 template <typename Kernel, typename DstHandler, typename ConstHandler, typename RegisterHandler, int32 NumInstancesPerOp>
 struct TUnaryKernel
@@ -1091,16 +1034,16 @@ struct FScalarKernelAcquireID
 	static VM_FORCEINLINE void Exec(FVectorVMContext& Context)
 	{
 		int32 DataSetIndex = VectorVM::DecodeU16(Context);
-		TArray<FDataSetMeta>& MetaTable = *Context.DataSetMetaTable;
+		TArrayView<FDataSetMeta> MetaTable = Context.DataSetMetaTable;
 		TArray<int32>&RESTRICT FreeIDTable = *MetaTable[DataSetIndex].FreeIDTable;
 
 		int32 Tag = MetaTable[DataSetIndex].IDAcquireTag;
 
 		int32 IDIndexReg = VectorVM::DecodeU16(Context);
-		int32*RESTRICT IDIndex = (int32*)(Context.RegisterTable[IDIndexReg]);
+		int32*RESTRICT IDIndex = (int32*)(Context.GetTempRegister(IDIndexReg));
 
 		int32 IDTagReg = VectorVM::DecodeU16(Context);
-		int32*RESTRICT IDTag = (int32*)(Context.RegisterTable[IDTagReg]);
+		int32*RESTRICT IDTag = (int32*)(Context.GetTempRegister(IDTagReg));
 
 		int32& NumFreeIDs = *MetaTable[DataSetIndex].NumFreeIDs;
 
@@ -1140,13 +1083,13 @@ struct FScalarKernelUpdateID
 		int32 InstanceIDRegisterIndex = VectorVM::DecodeU16(Context);
 		int32 InstanceIndexRegisterIndex = VectorVM::DecodeU16(Context);
 
-		TArray<FDataSetMeta>& MetaTable = *Context.DataSetMetaTable;
+		TArrayView<FDataSetMeta> MetaTable = Context.DataSetMetaTable;
 
 		TArray<int32>&RESTRICT IDTable = *MetaTable[DataSetIndex].IDTable;
 		int32 InstanceOffset = MetaTable[DataSetIndex].InstanceOffset + Context.StartInstance;
 
-		int32*RESTRICT IDRegister = (int32*)(Context.RegisterTable[InstanceIDRegisterIndex]);
-		int32*RESTRICT IndexRegister = (int32*)(Context.RegisterTable[InstanceIndexRegisterIndex]);
+		int32*RESTRICT IDRegister = (int32*)(Context.GetTempRegister(InstanceIDRegisterIndex));
+		int32*RESTRICT IndexRegister = (int32*)(Context.GetTempRegister(InstanceIndexRegisterIndex));
 		
 		FDataSetThreadLocalTempData& DataSetTempData = Context.ThreadLocalTempData[DataSetIndex];
 
@@ -1191,9 +1134,8 @@ struct FVectorKernelReadInput
 		int32 DestRegisterIdx = VectorVM::DecodeU16(Context);
 		int32 Loops = Align(Context.NumInstances, InstancesPerVector) / InstancesPerVector;
 
-		VectorRegister* DestReg = (VectorRegister*)(Context.RegisterTable[DestRegisterIdx]);
-		int32 DataSetOffset = Context.DataSetOffsetTable[DataSetIndex];  //TODO: we'll need a different way of doing this for a proper kernel; might need a specific offset handler
-		VectorRegister* InputReg = (VectorRegister*)((T*)(Context.RegisterTable[InputRegisterIdx + DataSetOffset]) + Context.StartInstance);
+		VectorRegister* DestReg = (VectorRegister*)(Context.GetTempRegister(DestRegisterIdx));
+		VectorRegister* InputReg = (VectorRegister*)(Context.GetInputRegister<T>(DataSetIndex, InputRegisterIdx) + Context.StartInstance);
 
 		//TODO: We can actually do some scalar loads into the first and final vectors to get around alignment issues and then use the aligned load for all others.
 		for (int32 i = 0; i < Loops; ++i)
@@ -1223,9 +1165,8 @@ struct FVectorKernelReadInputNoAdvance
 		int32 DestRegisterIdx = VectorVM::DecodeU16(Context);
 		int32 Loops = Align(Context.NumInstances, InstancesPerVector) / InstancesPerVector;
 
-		VectorRegister* DestReg = (VectorRegister*)(Context.RegisterTable[DestRegisterIdx]);
-		int32 DataSetOffset = Context.DataSetOffsetTable[DataSetIndex];  //TODO: we'll need a different way of doing this for a proper kernel; might need a specific offset handler
-		VectorRegister* InputReg = (VectorRegister*)((T*)(Context.RegisterTable[InputRegisterIdx + DataSetOffset]) );
+		VectorRegister* DestReg = (VectorRegister*)(Context.GetTempRegister(DestRegisterIdx));
+		VectorRegister* InputReg = (VectorRegister*)(Context.GetInputRegister<T>(DataSetIndex, InputRegisterIdx));
 
 		//TODO: We can actually do some scalar loads into the first and final vectors to get around alignment issues and then use the aligned load for all others.
 		for (int32 i = 0; i < Loops; ++i)
@@ -1262,23 +1203,6 @@ struct FVectorKernelReadInputNoAdvance
 // 	}
 // };
 
-//Needs it's own handler as the output registers are indexed absolutely rather than incrementing in advance().
-template<typename T>
-struct FOutputRegisterHandler : public FRegisterHandlerBase
-{
-	T* Register;
-	FOutputRegisterHandler(FVectorVMContext& Context, uint32 DataSetOffset)
-		: FRegisterHandlerBase(Context)
-		, Register((T*)Context.RegisterTable[RegisterIndex+DataSetOffset])
-	{}
-
-	VM_FORCEINLINE void Advance() { }
-	VM_FORCEINLINE T Get() { return *Register; }
-	VM_FORCEINLINE T*RESTRICT GetDest() { return Register; }
-	VM_FORCEINLINE T*RESTRICT GetDestAndAdvance() { return Register; }
-	VM_FORCEINLINE T GetAndAdvance() { return *Register; }
-};
-
 /** Special kernel for writing to a specific output register. */
 template<typename T>
 struct FScalarKernelWriteOutputIndexed
@@ -1288,17 +1212,36 @@ struct FScalarKernelWriteOutputIndexed
 		uint32 SrcOpTypes = VectorVM::DecodeSrcOperandTypes(Context);
 		switch (SrcOpTypes)
 		{
-		case SRCOP_RRR: TTrinaryOutputKernelHandler<FScalarKernelWriteOutputIndexed, FOutputRegisterHandler<T>, FDataSetOffsetHandler, FRegisterHandler<int32>, FRegisterHandler<T>, 1>::Exec(Context); break;
-		case SRCOP_RRC:	TTrinaryOutputKernelHandler<FScalarKernelWriteOutputIndexed, FOutputRegisterHandler<T>, FDataSetOffsetHandler, FRegisterHandler<int32>, FConstantHandler<T>, 1>::Exec(Context); break;
+		case SRCOP_RRR: DoKernel<FRegisterHandler<T>>(Context); break;
+		case SRCOP_RRC:	DoKernel<FConstantHandler<T>>(Context); break;
 		default: check(0); break;
 		};
 	}
 
-	static VM_FORCEINLINE void DoKernel(FVectorVMContext& Context, T* RESTRICT Dst, int32 SetOffset, int32 Index, T Data)
+	template<typename DataHandlerType>
+	static VM_FORCEINLINE void DoKernel(FVectorVMContext& Context)
 	{
-		if (Index != INDEX_NONE)
+		int32 DataSetIndex = VectorVM::DecodeU16(Context);
+
+		int32 DestIndexRegisterIdx = VectorVM::DecodeU16(Context);
+		T* DestIndexReg = (T*)(Context.GetTempRegister(DestIndexRegisterIdx));
+
+		DataHandlerType DataHandler(Context);
+
+		int32 DestRegisterIdx = VectorVM::DecodeU16(Context);
+		T* DestReg = Context.GetOutputRegister<T>(DataSetIndex, DestRegisterIdx);
+
+		for (int32 i = 0; i < Context.NumInstances; ++i)
 		{
-			Dst[Index] = Data;//TODO: On sse4 we can use _mm_stream_ss here.
+			int32 DestIndex = *DestIndexReg;
+			if (DestIndex != INDEX_NONE)
+			{
+				DestReg[DestIndex] = DataHandler.Get();
+			}
+
+			++DestIndexReg;
+			DataHandler.Advance();
+			//We don't increment the dest as we index into it directly.
 		}
 	}
 };
@@ -1307,7 +1250,7 @@ struct FDataSetCounterHandler
 {
 	int32* Counter;
 	FDataSetCounterHandler(FVectorVMContext& Context)
-		: Counter(Context.DataSetIndexTable + VectorVM::DecodeU16(Context))
+		: Counter(&Context.GetDataSetMeta(VectorVM::DecodeU16(Context)).DataSetAccessIndex)
 	{}
 
 	VM_FORCEINLINE void Advance() { }
@@ -1820,12 +1763,8 @@ void VectorVM::Init()
 void VectorVM::Exec(
 	uint8 const* Code,
 	int32 NumTempRegisters,
-	uint8** InputRegisters,
-	int32 NumInputRegisters,
-	uint8** OutputRegisters,
-	int32 NumOutputRegisters,
 	uint8 const* ConstantTable,
-	TArray<FDataSetMeta> &DataSetMetaTable,
+	TArrayView<FDataSetMeta> DataSetMetaTable,
 	FVMExternalFunction* ExternalFunctionTable,
 	void** UserPtrTable,
 	int32 NumInstances
@@ -1837,19 +1776,6 @@ void VectorVM::Exec(
 	//TRACE_CPUPROFILER_EVENT_SCOPE("VMExec");
 	SCOPE_CYCLE_COUNTER(STAT_VVMExec);
 
-	// table of index counters, one for each data set
-	TArray<int32, TInlineAllocator<16>> DataSetIndexTable;
-	TArray<int32, TInlineAllocator<16>> DataSetOffsetTable;
-
-	// map secondary data sets and fill in the offset table into the register table
-	//
-	for (int32 Idx = 0; Idx < DataSetMetaTable.Num(); Idx++)
-	{
-		uint32 DataSetOffset = DataSetMetaTable[Idx].RegisterOffset;
-		DataSetOffsetTable.Add(DataSetOffset);
-		DataSetIndexTable.Add(DataSetMetaTable[Idx].DataSetAccessIndex);	// prime counter index table with the data set offset; will be incremented with every write for each instance
-	}
-
 	int32 MaxInstances = FMath::Min(GParallelVVMInstancesPerChunk, NumInstances);
 	int32 NumChunks = (NumInstances / GParallelVVMInstancesPerChunk) + 1;
 	int32 ChunksPerBatch = (GbParallelVVM != 0 && FApp::ShouldUseThreadingForPerformance()) ? GParallelVVMChunksPerBatch : NumChunks;
@@ -1861,8 +1787,7 @@ void VectorVM::Exec(
 		//SCOPE_CYCLE_COUNTER(STAT_VVMExecChunk);
 
 		FVectorVMContext& Context = FVectorVMContext::Get();
-		Context.PrepareForExec(InputRegisters, OutputRegisters, NumTempRegisters, NumInputRegisters, NumOutputRegisters, ConstantTable, DataSetIndexTable.GetData(), DataSetOffsetTable.GetData(), DataSetOffsetTable.Num(),
-			ExternalFunctionTable, UserPtrTable, DataSetMetaTable, MaxInstances
+		Context.PrepareForExec(NumTempRegisters, ConstantTable, ExternalFunctionTable, UserPtrTable, DataSetMetaTable, MaxInstances
 #if STATS
 			, &StatScopes
 #endif
@@ -2016,12 +1941,6 @@ void VectorVM::Exec(
 	};
 
 	ParallelFor(NumBatches, ExecChunkBatch, GbParallelVVM == 0 || !bParallel);
-
-	// write back data set access indices, so we know how much was written to each data set
-	for (int32 Idx = 0; Idx < DataSetMetaTable.Num(); Idx++)
-	{
-		DataSetMetaTable[Idx].DataSetAccessIndex = DataSetIndexTable[Idx];	
-	}
 }
 
 uint8 VectorVM::GetNumOpCodes()
diff --git a/Engine/Source/Runtime/VectorVM/Public/VectorVM.h b/Engine/Source/Runtime/VectorVM/Public/VectorVM.h
--- a/Engine/Source/Runtime/VectorVM/Public/VectorVM.h
+++ b/Engine/Source/Runtime/VectorVM/Public/VectorVM.h
@@ -154,7 +154,9 @@ enum class EVectorVMOp : uint8
 //Some require RWBuffer like support.
 struct FDataSetMeta
 {
-	int32 RegisterOffset;
+	uint8*RESTRICT*RESTRICT InputRegisters;
+	uint8*RESTRICT*RESTRICT OutputRegisters;
+
 	int32 DataSetAccessIndex;	// index for individual elements of this set
 
 	int32 InstanceOffset;		// offset of the first instance processed 
@@ -179,16 +181,45 @@ struct FDataSetMeta
 	FORCEINLINE void LockFreeTable();
 	FORCEINLINE void UnlockFreeTable();
 
-	FDataSetMeta(int32 InRegisterOffset, int32 InInstanceOffset, TArray<int32>* InIDTable, TArray<int32>* InFreeIDTable, int32* InNumFreeIDs, int32* InMaxUsedID, int32 InIDAcquireTag)
-		: RegisterOffset(InRegisterOffset), DataSetAccessIndex(INDEX_NONE), InstanceOffset(InInstanceOffset)
-		, IDTable(InIDTable), FreeIDTable(InFreeIDTable), NumFreeIDs(InNumFreeIDs), MaxUsedID(InMaxUsedID), IDAcquireTag(InIDAcquireTag)
+	FDataSetMeta()
+		: InputRegisters(nullptr)
+		, OutputRegisters(nullptr)
+		, DataSetAccessIndex(INDEX_NONE)
+		, InstanceOffset(INDEX_NONE)
+		, IDTable(nullptr)
+		, FreeIDTable(nullptr)
+		, NumFreeIDs(nullptr)
+		, MaxUsedID(nullptr)
+		, IDAcquireTag(INDEX_NONE)
+	{
+	}	
+
+	FORCEINLINE void Reset()
 	{
+		InputRegisters = nullptr;
+		OutputRegisters = nullptr;
+		DataSetAccessIndex = INDEX_NONE;
+		InstanceOffset = INDEX_NONE;
+		IDTable = nullptr;
+		FreeIDTable = nullptr;
+		NumFreeIDs = nullptr;
+		MaxUsedID = nullptr;
+		IDAcquireTag = INDEX_NONE;
 	}
 
-	FDataSetMeta() 
-		: RegisterOffset(0), DataSetAccessIndex(INDEX_NONE), InstanceOffset(0)
-		, IDTable(nullptr), FreeIDTable(nullptr), NumFreeIDs(nullptr), MaxUsedID(nullptr), IDAcquireTag(0)
-	{}
+	FORCEINLINE void Init(uint8*RESTRICT *RESTRICT InInputRegisters, uint8*RESTRICT *RESTRICT InOutputRegisters, int32 InInstanceOffset, TArray<int32>* InIDTable, TArray<int32>* InFreeIDTable, int32* InNumFreeIDs, int32* InMaxUsedID, int32 InIDAcquireTag)
+	{
+		InputRegisters = InInputRegisters;
+		OutputRegisters = InOutputRegisters;
+
+		DataSetAccessIndex = INDEX_NONE;
+		InstanceOffset = InInstanceOffset;
+		IDTable = InIDTable;
+		FreeIDTable = InFreeIDTable;
+		NumFreeIDs = InNumFreeIDs;
+		MaxUsedID = InMaxUsedID;
+		IDAcquireTag = InIDAcquireTag;
+	}
 
 private:
 	// Non-copyable and non-movable
@@ -198,28 +229,19 @@ private:
 	FDataSetMeta& operator=(const FDataSetMeta&) = delete;
 };
 
-namespace VectorVM
-{
-	/** Constants. */
-	enum
-	{
-		NumTempRegisters = 400,
-		MaxInputRegisters = 400,
-		MaxOutputRegisters = MaxInputRegisters,
-		MaxConstants = 256,
-		FirstTempRegister = 0,
-		FirstInputRegister = NumTempRegisters,
-		FirstOutputRegister = FirstInputRegister + MaxInputRegisters,
-		MaxRegisters = NumTempRegisters + MaxInputRegisters + MaxOutputRegisters + MaxConstants,
-	};
-}
-
 //Data the VM will keep on each dataset locally per thread which is then thread safely pushed to it's destination at the end of execution.
 struct FDataSetThreadLocalTempData
 {
 	FDataSetThreadLocalTempData()
-		:MaxID(INDEX_NONE)
-	{}
+	{
+		Reset();
+	}
+
+	FORCEINLINE void Reset()
+	{
+		IDsToFree.Reset();
+		MaxID = INDEX_NONE;
+	}
 
 	TArray<int32> IDsToFree;
 	int32 MaxID;
@@ -241,14 +263,7 @@ struct FVectorVMContext : TThreadSingleton<FVectorVMContext>
 	uint8 const* RESTRICT ConstantTable;
 	/** Num temp registers required by this script. */
 	int32 NumTempRegisters;
-	/** Num input required by this script. */
-	int32 NumInputRegisters;
-	/** Num output required by this script. */
-	int32 NumOutputRegisters;
-	/** Pointer to the data set index counter table */
-	int32* RESTRICT DataSetIndexTable;
-	int32* RESTRICT DataSetOffsetTable;
-	int32 NumSecondaryDataSets;
+
 	/** Pointer to the shared data table. */
 	FVMExternalFunction* RESTRICT ExternalFunctionTable;
 	/** Table of user pointers.*/
@@ -259,7 +274,7 @@ struct FVectorVMContext : TThreadSingleton<FVectorVMContext>
 	int32 StartInstance;
 
 	/** Array of meta data on data sets. TODO: This struct should be removed and all features it contains be handled by more general vm ops and the compiler's knowledge of offsets etc. */
-	TArray<FDataSetMeta>* RESTRICT DataSetMetaTable;
+	TArrayView<FDataSetMeta> DataSetMetaTable;
 
 	TArray<FDataSetThreadLocalTempData> ThreadLocalTempData;
 
@@ -269,7 +284,8 @@ struct FVectorVMContext : TThreadSingleton<FVectorVMContext>
 #endif
 
 	TArray<uint8, TAlignedHeapAllocator<VECTOR_WIDTH_BYTES>> TempRegTable;
-	uint8*RESTRICT RegisterTable[VectorVM::MaxRegisters];
+	uint32 TempRegisterSize;
+	uint32 TempBufferSize;
 
 	/** Thread local random stream for use in external functions needing non-deterministic randoms. */
 	FRandomStream RandStream;
@@ -280,18 +296,11 @@ struct FVectorVMContext : TThreadSingleton<FVectorVMContext>
 	FVectorVMContext();
 
 	void PrepareForExec(
-		uint8*RESTRICT*RESTRICT InputRegisters,
-		uint8*RESTRICT*RESTRICT OutputRegisters,
 		int32 InNumTempRegisters,
-		int32 NumInputRegisters,
-		int32 NumOutputRegisters,
 		const uint8* InConstantTable,
-		int32 *InDataSetIndexTable,
-		int32 *InDataSetOffsetTable,
-		int32 InNumSecondaryDatasets,
 		FVMExternalFunction* InExternalFunctionTable,
 		void** InUserPtrTable,
-		TArray<FDataSetMeta>& RESTRICT InDataSetMetaTable,
+		TArrayView<FDataSetMeta> InDataSetMetaTable,
 		int32 MaxNumInstances
 #if STATS
 		, const TArray<TStatId>* InStatScopes
@@ -309,6 +318,22 @@ struct FVectorVMContext : TThreadSingleton<FVectorVMContext>
 		RandCounters.Reset();
 		RandCounters.SetNumZeroed(InNumInstances);
 	}
+
+	FORCEINLINE FDataSetMeta& GetDataSetMeta(int32 DataSetIndex) { return DataSetMetaTable[DataSetIndex]; }
+	FORCEINLINE uint8 * RESTRICT GetTempRegister(int32 RegisterIndex) { return TempRegTable.GetData() + TempRegisterSize * RegisterIndex; }
+	template<typename T>
+	FORCEINLINE T* RESTRICT GetInputRegister(int32 DataSetIndex, int32 RegisterIndex) 
+	{
+		FDataSetMeta& Meta = GetDataSetMeta(DataSetIndex);
+		return ((T*)Meta.InputRegisters[RegisterIndex]) + Meta.InstanceOffset;
+	}
+	template<typename T>
+	FORCEINLINE T* RESTRICT GetOutputRegister(int32 DataSetIndex, int32 RegisterIndex) 
+	{ 
+		FDataSetMeta& Meta = GetDataSetMeta(DataSetIndex);
+		return  ((T*)Meta.OutputRegisters[RegisterIndex]) + Meta.InstanceOffset;
+	}
+
 };
 
 namespace VectorVM
@@ -329,12 +354,8 @@ namespace VectorVM
 	VECTORVM_API void Exec(
 		uint8 const* Code,
 		int32 NumTempRegisters,
-		uint8** InputRegisters,
-		int32 NumInputRegisters,
-		uint8** OutputRegisters,
-		int32 NumOutputRegisters,
 		uint8 const* ConstantTable,
-		TArray<FDataSetMeta> &DataSetMetaTable,
+		TArrayView<FDataSetMeta> DataSetMetaTable,
 		FVMExternalFunction* ExternalFunctionTable,
 		void** UserPtrTable,
 		int32 NumInstances
@@ -415,7 +436,7 @@ namespace VectorVM
 		void Init(FVectorVMContext& Context)
 		{
 			InputOffset = DecodeU16(Context);
-			InputPtr = IsConstant() ? (T*)(Context.ConstantTable + GetOffset()) : (T*)Context.RegisterTable[GetOffset()];
+			InputPtr = IsConstant() ? (T*)(Context.ConstantTable + GetOffset()) : (T*)Context.GetTempRegister(GetOffset());
 			AdvanceOffset = IsConstant() ? 0 : 1;
 		}
 
@@ -455,8 +476,8 @@ namespace VectorVM
 		{
 			if (IsValid())
 			{
-				check(RegisterIndex < VectorVM::MaxRegisters);
-				Register = (T*)Context.RegisterTable[RegisterIndex];
+				checkSlow(RegisterIndex < Context.NumTempRegisters);
+				Register = (T*)Context.GetTempRegister(RegisterIndex);
 			}
 			else
 			{
