- Write VectorVM code in LittleEndian as opposed to BigEndian, this avoids runtime byte swapping
- Read ByteCode directly if unalgned loads are supported to condense 3 opts into 1 when reading a uint16
- Pre-calculate instance loops for VECTOR_WIDTH_FLOATS to avoid each vector op having to round / divide instance counts
- Added runtime optimization of the VM script, this currently boils down to a function call per VM invoke + storing the data required
- Use vm.OptimizeVMByteCode to enable optimized code generation
- Use vm.UseOptimizedVMByteCode to enable running optimized code rather the traditional byte code

#rnx

diff --git a/Engine/Plugins/FX/Niagara/Source/Niagara/Classes/NiagaraScript.h b/Engine/Plugins/FX/Niagara/Source/Niagara/Classes/NiagaraScript.h
--- a/Engine/Plugins/FX/Niagara/Source/Niagara/Classes/NiagaraScript.h
+++ b/Engine/Plugins/FX/Niagara/Source/Niagara/Classes/NiagaraScript.h
@@ -225,6 +225,10 @@ public:
 	UPROPERTY()
 	TArray<uint8> ByteCode;
 
+	/** Runtime optimized byte code, specific to the system we are running on, currently can not be serialized */
+	UPROPERTY(transient)
+	TArray<uint8> OptimizedByteCode;
+
 	/** Number of temp registers used by this script. */
 	UPROPERTY()
 	int32 NumTempRegisters;
diff --git a/Engine/Plugins/FX/Niagara/Source/Niagara/Private/NiagaraScript.cpp b/Engine/Plugins/FX/Niagara/Source/Niagara/Private/NiagaraScript.cpp
--- a/Engine/Plugins/FX/Niagara/Source/Niagara/Private/NiagaraScript.cpp
+++ b/Engine/Plugins/FX/Niagara/Source/Niagara/Private/NiagaraScript.cpp
@@ -29,6 +29,9 @@
 #include "UObject/FortniteMainBranchObjectVersion.h"
 #include "UObject/RenderingObjectVersion.h"
 
+#include "NiagaraFunctionLibrary.h"
+#include "VectorVM.h"
+
 DECLARE_STATS_GROUP(TEXT("Niagara Detailed"), STATGROUP_NiagaraDetailed, STATCAT_Advanced);
 
 FNiagaraScriptDebuggerInfo::FNiagaraScriptDebuggerInfo() : bWaitForGPU(false), FrameLastWriteId(-1), bWritten(false)
@@ -728,6 +731,21 @@ void UNiagaraScript::PostLoad()
 	GenerateStatScopeIDs();
 #endif
 
+	// Optimize the VM script for runtime usage
+	if ( FPlatformProperties::RequiresCookedData() && CachedScriptVM.IsValid() )
+	{
+		TArray<uint8, TInlineAllocator<32>> ExternalFunctionRegisterCounts;
+		ExternalFunctionRegisterCounts.Reserve(CachedScriptVM.CalledVMExternalFunctions.Num());
+
+		for ( const FVMExternalFunctionBindingInfo FunctionBindingInfo : CachedScriptVM.CalledVMExternalFunctions )
+		{
+			const uint8 RegisterCount = FunctionBindingInfo.GetNumInputs() + FunctionBindingInfo.GetNumOutputs();
+			ExternalFunctionRegisterCounts.Add(RegisterCount);
+		}
+
+		VectorVM::OptimizeByteCode(CachedScriptVM.ByteCode.GetData(), CachedScriptVM.OptimizedByteCode, MakeArrayView(ExternalFunctionRegisterCounts));
+	}
+
 	//FNiagaraUtilities::DumpHLSLText(RapidIterationParameters.ToString(), *GetPathName());
 }
 
diff --git a/Engine/Plugins/FX/Niagara/Source/Niagara/Private/NiagaraScriptExecutionContext.cpp b/Engine/Plugins/FX/Niagara/Source/Niagara/Private/NiagaraScriptExecutionContext.cpp
--- a/Engine/Plugins/FX/Niagara/Source/Niagara/Private/NiagaraScriptExecutionContext.cpp
+++ b/Engine/Plugins/FX/Niagara/Source/Niagara/Private/NiagaraScriptExecutionContext.cpp
@@ -188,6 +188,7 @@ bool FNiagaraScriptExecutionContext::Execute(uint32 NumInstances)
 		const FNiagaraVMExecutableData& ExecData = Script->GetVMExecutableData();
 		VectorVM::Exec(
 			ExecData.ByteCode.GetData(),
+			ExecData.OptimizedByteCode.Num() > 0 ? ExecData.OptimizedByteCode.GetData() : nullptr,
 			ExecData.NumTempRegisters,
 			Parameters.GetParameterDataArray().GetData(),
 			DataSetMetaTable,
diff --git a/Engine/Plugins/FX/Niagara/Source/NiagaraCore/Private/NiagaraCustomVersion.cpp b/Engine/Plugins/FX/Niagara/Source/NiagaraCore/Private/NiagaraCustomVersion.cpp
--- a/Engine/Plugins/FX/Niagara/Source/NiagaraCore/Private/NiagaraCustomVersion.cpp
+++ b/Engine/Plugins/FX/Niagara/Source/NiagaraCore/Private/NiagaraCustomVersion.cpp
@@ -8,4 +8,4 @@ const FGuid FNiagaraCustomVersion::GUID(0xFCF57AFA, 0x50764283, 0xB9A9E658, 0xFF
 // Register the custom version with core
 FCustomVersionRegistration GRegisterNiagaraCustomVersion(FNiagaraCustomVersion::GUID, FNiagaraCustomVersion::LatestVersion, TEXT("NiagaraVer"));
 
-const FGuid FNiagaraCustomVersion::LatestScriptCompileVersion(0x0F93E860, 0x6A30224B, 0xA9000879, 0xA7693A89);
\ No newline at end of file
+const FGuid FNiagaraCustomVersion::LatestScriptCompileVersion(0x79007C85, 0x2A9CA54A, 0xA7DABF24, 0x30A4EB07);
\ No newline at end of file
diff --git a/Engine/Source/Runtime/VectorVM/Private/VectorVM.cpp b/Engine/Source/Runtime/VectorVM/Private/VectorVM.cpp
--- a/Engine/Source/Runtime/VectorVM/Private/VectorVM.cpp
+++ b/Engine/Source/Runtime/VectorVM/Private/VectorVM.cpp
@@ -98,27 +98,95 @@ static FAutoConsoleVariableRef CVarParallelVVMInstancesPerChunk(
 	ECVF_ReadOnly
 );
 
+static int32 GbOptimizeVMByteCode = 0;
+static FAutoConsoleVariableRef CVarbOptimizeVMByteCode(
+	TEXT("vm.OptimizeVMByteCode"),
+	GbOptimizeVMByteCode,
+	TEXT("If > 0 vector VM code optimization will be enabled at runtime.\n"),
+	ECVF_Default
+);
+
+static int32 GbUseOptimizedVMByteCode = 0;
+static FAutoConsoleVariableRef CVarbUseOptimizedVMByteCode(
+	TEXT("vm.UseOptimizedVMByteCode"),
+	GbUseOptimizedVMByteCode,
+	TEXT("If > 0 optimized vector VM code will be excuted at runtime.\n"),
+	ECVF_Default
+);
+
+//////////////////////////////////////////////////////////////////////////
+//  VM Code Optimizer Context
+
+typedef void(*FVectorVMExecFunction)(FVectorVMContext&);
+
+struct FVectorVMCodeOptimizerContext
+{
+	explicit FVectorVMCodeOptimizerContext(FVectorVMContext& InBaseContext, const uint8* ByteCode, TArray<uint8>& InOptimizedCode, TArrayView<uint8> InExternalFunctionRegisterCounts)
+		: BaseContext(InBaseContext)
+		, OptimizedCode(InOptimizedCode)
+		, ExternalFunctionRegisterCounts(InExternalFunctionRegisterCounts)
+	{
+#if STATS
+		BaseContext.PrepareForExec(0, nullptr, nullptr, nullptr, TArrayView<FDataSetMeta>(), 0, false, &EmptyStatScopes);
+#else
+		BaseContext.PrepareForExec(0, nullptr, nullptr, nullptr, TArrayView<FDataSetMeta>(), 0, false);
+#endif
+		BaseContext.PrepareForChunk(ByteCode, 0, 0);
+	}
+	FVectorVMCodeOptimizerContext(const FVectorVMCodeOptimizerContext&) = delete;
+	FVectorVMCodeOptimizerContext(const FVectorVMCodeOptimizerContext&&) = delete;
+
+	template<uint32 InstancesPerOp>
+	int32 GetNumLoops() const { return 0; }
+
+	FORCEINLINE uint8 DecodeU8() { return BaseContext.DecodeU8(); }
+	FORCEINLINE uint16 DecodeU16() { return BaseContext.DecodeU16(); }
+	FORCEINLINE uint32 DecodeU32() { return BaseContext.DecodeU32(); }
+	FORCEINLINE uint64 DecodeU64() { return BaseContext.DecodeU64(); }
+
+	//-TODO: Support unaligned writes
+	template<typename T>
+	void Write(const T& v)
+	{
+		reinterpret_cast<T&>(OptimizedCode[OptimizedCode.AddUninitialized(sizeof(T))]) = v;
+	}
+
+	FVectorVMContext&		BaseContext;
+	TArray<uint8>&			OptimizedCode;
+	const TArrayView<uint8>	ExternalFunctionRegisterCounts;
+	const int32				StartInstance = 0;
+#if STATS
+	TArray<TStatId>			EmptyStatScopes;
+#endif
+};
+
 //////////////////////////////////////////////////////////////////////////
 //  Constant Handlers
 
 struct FConstantHandlerBase
 {
-	uint16 ConstantIndex;
+	const uint16 ConstantIndex;
 	FConstantHandlerBase(FVectorVMContext& Context)
-		: ConstantIndex(VectorVM::DecodeU16(Context))
+		: ConstantIndex(Context.DecodeU16())
 	{}
 
 	FORCEINLINE void Advance() { }
+
+	static void Optimize(FVectorVMCodeOptimizerContext& Context)
+	{
+		Context.Write(Context.DecodeU16());
+	}
 };
 
 template<typename T>
 struct FConstantHandler : public FConstantHandlerBase
 {
-	T Constant;
+	const T Constant;
 	FConstantHandler(FVectorVMContext& Context)
 		: FConstantHandlerBase(Context)
 		, Constant(*((T*)(Context.ConstantTable + ConstantIndex)))
 	{}
+
 	FORCEINLINE const T& Get() { return Constant; }
 	FORCEINLINE const T& GetAndAdvance() { return Constant; }
 };
@@ -126,11 +194,12 @@ struct FConstantHandler : public FConstantHandlerBase
 template<>
 struct FConstantHandler<VectorRegister> : public FConstantHandlerBase
 {
-	VectorRegister Constant;
+	const VectorRegister Constant;
 	FConstantHandler(FVectorVMContext& Context)
 		: FConstantHandlerBase(Context)
 		, Constant(VectorLoadFloat1(&Context.ConstantTable[ConstantIndex]))
 	{}
+
 	FORCEINLINE const VectorRegister Get() { return Constant; }
 	FORCEINLINE const VectorRegister GetAndAdvance() { return Constant; }
 };
@@ -138,11 +207,12 @@ struct FConstantHandler<VectorRegister> : public FConstantHandlerBase
 template<>
 struct FConstantHandler<VectorRegisterInt> : public FConstantHandlerBase
 {
-	VectorRegisterInt Constant;
+	const VectorRegisterInt Constant;
 	FConstantHandler(FVectorVMContext& Context)
 		: FConstantHandlerBase(Context)
 		, Constant(VectorIntLoad1(&Context.ConstantTable[ConstantIndex]))
 	{}
+
 	FORCEINLINE const VectorRegisterInt Get() { return Constant; }
 	FORCEINLINE const VectorRegisterInt GetAndAdvance() { return Constant; }
 };
@@ -154,11 +224,15 @@ struct FConstantHandler<VectorRegisterInt> : public FConstantHandlerBase
 
 struct FRegisterHandlerBase
 {
-	int32 RegisterIndex;
+	const int32 RegisterIndex;
 	FORCEINLINE FRegisterHandlerBase(FVectorVMContext& Context)
-		: RegisterIndex(VectorVM::DecodeU16(Context))
+		: RegisterIndex(Context.DecodeU16())
 	{}
 
+	static void Optimize(FVectorVMCodeOptimizerContext& Context)
+	{
+		Context.Write(Context.DecodeU16());
+	}
 };
 
 template<typename T>
@@ -171,6 +245,7 @@ public:
 		: FRegisterHandlerBase(Context)
 		, Register((T*)Context.GetTempRegister(RegisterIndex))
 	{}
+
 	FORCEINLINE const T Get() { return *Register; }
 	FORCEINLINE T* GetDest() { return Register; }
 	FORCEINLINE void Advance() { ++Register; }
@@ -208,7 +283,8 @@ void FVectorVMContext::PrepareForExec(
 	FVMExternalFunction* InExternalFunctionTable,
 	void** InUserPtrTable,
 	TArrayView<FDataSetMeta> InDataSetMetaTable,
-	int32 MaxNumInstances
+	int32 MaxNumInstances,
+	bool bInParallelExecution
 #if STATS
 	, const TArray<TStatId>* InStatScopes
 #endif
@@ -235,6 +311,8 @@ void FVectorVMContext::PrepareForExec(
 		TLSTempData.Reset();
 	}
 	ThreadLocalTempData.SetNum(DataSetMetaTable.Num());
+
+	bIsParallelExecution = bInParallelExecution;
 }
 
 void FVectorVMContext::FinishExec()
@@ -265,19 +343,26 @@ void FVectorVMContext::FinishExec()
 		}
 
 		//Also update the max ID seen. This should be the ONLY place in the VM we update this max value.
-		volatile int32* MaxUsedID = DataSetMetaTable[DataSetIndex].MaxUsedID;
-		int32 LocalMaxUsedID;
-		do
+		if ( bIsParallelExecution )
 		{
-			LocalMaxUsedID = *MaxUsedID;
-			if (LocalMaxUsedID >= Data.MaxID)
+			volatile int32* MaxUsedID = DataSetMetaTable[DataSetIndex].MaxUsedID;
+			int32 LocalMaxUsedID;
+			do
 			{
-				break;
-			}
-		} while (FPlatformAtomics::InterlockedCompareExchange(MaxUsedID, Data.MaxID, LocalMaxUsedID) != LocalMaxUsedID);
-
-		*MaxUsedID = FMath::Max(*MaxUsedID, Data.MaxID);
+				LocalMaxUsedID = *MaxUsedID;
+				if (LocalMaxUsedID >= Data.MaxID)
+				{
+					break;
+				}
+			} while (FPlatformAtomics::InterlockedCompareExchange(MaxUsedID, Data.MaxID, LocalMaxUsedID) != LocalMaxUsedID);
 
+			*MaxUsedID = FMath::Max(*MaxUsedID, Data.MaxID);
+		}
+		else
+		{
+			int32* MaxUsedID = DataSetMetaTable[DataSetIndex].MaxUsedID;
+			*MaxUsedID = FMath::Max(*MaxUsedID, Data.MaxID);
+		}
 	}
 }
 
@@ -295,13 +380,20 @@ uint8 VectorVM::CreateSrcOperandMask(EVectorVMOperandLocation Type0, EVectorVMOp
 template<typename Kernel, typename DstHandler, typename Arg0Handler, uint32 NumInstancesPerOp>
 struct TUnaryKernelHandler
 {
+	static void Optimize(FVectorVMCodeOptimizerContext& Context)
+	{
+		Context.Write<FVectorVMExecFunction>(Exec);
+		Arg0Handler::Optimize(Context);
+		DstHandler::Optimize(Context);
+	}
+
 	static VM_FORCEINLINE void Exec(FVectorVMContext& Context)
 	{
 		Arg0Handler Arg0(Context);
 		DstHandler Dst(Context);
 
-		int32 LoopInstances = Align(Context.NumInstances, NumInstancesPerOp) / NumInstancesPerOp;
-		for (int32 i = 0; i < LoopInstances; ++i)
+		const int32 Loops = Context.GetNumLoops<NumInstancesPerOp>();
+		for (int32 i = 0; i < Loops; ++i)
 		{
 			Kernel::DoKernel(Context, Dst.GetDestAndAdvance(), Arg0.GetAndAdvance());
 		}
@@ -311,6 +403,14 @@ struct TUnaryKernelHandler
 template<typename Kernel, typename DstHandler, typename Arg0Handler, typename Arg1Handler, int32 NumInstancesPerOp>
 struct TBinaryKernelHandler
 {
+	static void Optimize(FVectorVMCodeOptimizerContext& Context)
+	{
+		Context.Write<FVectorVMExecFunction>(Exec);
+		Arg0Handler::Optimize(Context);
+		Arg1Handler::Optimize(Context);
+		DstHandler::Optimize(Context);
+	}
+
 	static void Exec(FVectorVMContext& Context)
 	{
 		Arg0Handler Arg0(Context); 
@@ -318,8 +418,8 @@ struct TBinaryKernelHandler
 
 		DstHandler Dst(Context);
 
-		int32 LoopInstances = Align(Context.NumInstances, NumInstancesPerOp) / NumInstancesPerOp;
-		for (int32 i = 0; i < LoopInstances; ++i)
+		const int32 Loops = Context.GetNumLoops<NumInstancesPerOp>();
+		for (int32 i = 0; i < Loops; ++i)
 		{
 			Kernel::DoKernel(Context, Dst.GetDestAndAdvance(), Arg0.GetAndAdvance(), Arg1.GetAndAdvance());
 		}
@@ -329,6 +429,15 @@ struct TBinaryKernelHandler
 template<typename Kernel, typename DstHandler, typename Arg0Handler, typename Arg1Handler, typename Arg2Handler, int32 NumInstancesPerOp>
 struct TTrinaryKernelHandler
 {
+	static void Optimize(FVectorVMCodeOptimizerContext& Context)
+	{
+		Context.Write<FVectorVMExecFunction>(Exec);
+		Arg0Handler::Optimize(Context);
+		Arg1Handler::Optimize(Context);
+		Arg2Handler::Optimize(Context);
+		DstHandler::Optimize(Context);
+	}
+
 	static void Exec(FVectorVMContext& Context)
 	{
 		Arg0Handler Arg0(Context);
@@ -337,8 +446,8 @@ struct TTrinaryKernelHandler
 
 		DstHandler Dst(Context);
 
-		int32 LoopInstances = Align(Context.NumInstances, NumInstancesPerOp) / NumInstancesPerOp;
-		for (int32 i = 0; i < LoopInstances; ++i)
+		const int32 Loops = Context.GetNumLoops<NumInstancesPerOp>();
+		for (int32 i = 0; i < Loops; ++i)
 		{
 			Kernel::DoKernel(Context, Dst.GetDestAndAdvance(), Arg0.GetAndAdvance(), Arg1.GetAndAdvance(), Arg2.GetAndAdvance());
 		}
@@ -350,9 +459,20 @@ struct TTrinaryKernelHandler
 template <typename Kernel, typename DstHandler, typename ConstHandler, typename RegisterHandler, int32 NumInstancesPerOp>
 struct TUnaryKernel
 {
+	static void Optimize(FVectorVMCodeOptimizerContext& Context)
+	{
+		const uint32 SrcOpTypes = Context.BaseContext.DecodeSrcOperandTypes();
+		switch (SrcOpTypes)
+		{
+		case SRCOP_RRR: TUnaryKernelHandler<Kernel, DstHandler, RegisterHandler, NumInstancesPerOp>::Optimize(Context); break;
+		case SRCOP_RRC:	TUnaryKernelHandler<Kernel, DstHandler, ConstHandler, NumInstancesPerOp>::Optimize(Context); break;
+		default: check(0); break;
+		};
+	}
+
 	static void Exec(FVectorVMContext& Context)
 	{
-		uint32 SrcOpTypes = VectorVM::DecodeSrcOperandTypes(Context);
+		const uint32 SrcOpTypes = Context.DecodeSrcOperandTypes();
 		switch (SrcOpTypes)
 		{
 		case SRCOP_RRR: TUnaryKernelHandler<Kernel, DstHandler, RegisterHandler, NumInstancesPerOp>::Exec(Context); break;
@@ -374,9 +494,22 @@ struct TUnaryVectorIntKernel : public TUnaryKernel<Kernel, FRegisterHandler<Vect
 template <typename Kernel, typename DstHandler, typename ConstHandler, typename RegisterHandler, uint32 NumInstancesPerOp>
 struct TBinaryKernel
 {
+	static void Optimize(FVectorVMCodeOptimizerContext& Context)
+	{
+		const uint32 SrcOpTypes = Context.BaseContext.DecodeSrcOperandTypes();
+		switch (SrcOpTypes)
+		{
+		case SRCOP_RRR: TBinaryKernelHandler<Kernel, DstHandler, RegisterHandler, RegisterHandler, NumInstancesPerOp>::Optimize(Context); break;
+		case SRCOP_RRC:	TBinaryKernelHandler<Kernel, DstHandler, ConstHandler, RegisterHandler, NumInstancesPerOp>::Optimize(Context); break;
+		case SRCOP_RCR: TBinaryKernelHandler<Kernel, DstHandler, RegisterHandler, ConstHandler, NumInstancesPerOp>::Optimize(Context); break;
+		case SRCOP_RCC:	TBinaryKernelHandler<Kernel, DstHandler, ConstHandler, ConstHandler, NumInstancesPerOp>::Optimize(Context); break;
+		default: check(0); break;
+		};
+	}
+
 	static void Exec(FVectorVMContext& Context)
 	{
-		uint32 SrcOpTypes = VectorVM::DecodeSrcOperandTypes(Context);
+		const uint32 SrcOpTypes = Context.DecodeSrcOperandTypes();
 		switch (SrcOpTypes)
 		{
 		case SRCOP_RRR: TBinaryKernelHandler<Kernel, DstHandler, RegisterHandler, RegisterHandler, NumInstancesPerOp>::Exec(Context); break;
@@ -398,9 +531,26 @@ struct TBinaryVectorIntKernel : public TBinaryKernel<Kernel, FRegisterHandler<Ve
 template <typename Kernel, typename DstHandler, typename ConstHandler, typename RegisterHandler, uint32 NumInstancesPerOp>
 struct TTrinaryKernel
 {
+	static void Optimize(FVectorVMCodeOptimizerContext& Context)
+	{
+		const uint32 SrcOpTypes = Context.BaseContext.DecodeSrcOperandTypes();
+		switch (SrcOpTypes)
+		{
+		case SRCOP_RRR: TTrinaryKernelHandler<Kernel, DstHandler, RegisterHandler, RegisterHandler, RegisterHandler, NumInstancesPerOp>::Optimize(Context); break;
+		case SRCOP_RRC:	TTrinaryKernelHandler<Kernel, DstHandler, ConstHandler, RegisterHandler, RegisterHandler, NumInstancesPerOp>::Optimize(Context); break;
+		case SRCOP_RCR: TTrinaryKernelHandler<Kernel, DstHandler, RegisterHandler, ConstHandler, RegisterHandler, NumInstancesPerOp>::Optimize(Context); break;
+		case SRCOP_RCC:	TTrinaryKernelHandler<Kernel, DstHandler, ConstHandler, ConstHandler, RegisterHandler, NumInstancesPerOp>::Optimize(Context); break;
+		case SRCOP_CRR: TTrinaryKernelHandler<Kernel, DstHandler, RegisterHandler, RegisterHandler, ConstHandler, NumInstancesPerOp>::Optimize(Context); break;
+		case SRCOP_CRC:	TTrinaryKernelHandler<Kernel, DstHandler, ConstHandler, RegisterHandler, ConstHandler, NumInstancesPerOp>::Optimize(Context); break;
+		case SRCOP_CCR: TTrinaryKernelHandler<Kernel, DstHandler, RegisterHandler, ConstHandler, ConstHandler, NumInstancesPerOp>::Optimize(Context); break;
+		case SRCOP_CCC:	TTrinaryKernelHandler<Kernel, DstHandler, ConstHandler, ConstHandler, ConstHandler, NumInstancesPerOp>::Optimize(Context); break;
+		default: check(0); break;
+		};
+	}
+
 	static void Exec(FVectorVMContext& Context)
 	{
-		uint32 SrcOpTypes = VectorVM::DecodeSrcOperandTypes(Context);
+		const uint32 SrcOpTypes = Context.DecodeSrcOperandTypes();
 		switch (SrcOpTypes)
 		{
 		case SRCOP_RRR: TTrinaryKernelHandler<Kernel, DstHandler, RegisterHandler, RegisterHandler, RegisterHandler, NumInstancesPerOp>::Exec(Context); break;
@@ -725,6 +875,12 @@ struct FVectorKernelSelect : public TTrinaryVectorKernel<FVectorKernelSelect>
 
 struct FVectorKernelExecutionIndex
 {
+	static void Optimize(FVectorVMCodeOptimizerContext& Context)
+	{
+		Context.Write<FVectorVMExecFunction>(Exec);
+		FRegisterHandler<VectorRegisterInt>::Optimize(Context);
+	}
+
 	static void VM_FORCEINLINE Exec(FVectorVMContext& Context)
 	{
 		static_assert(VECTOR_WIDTH_FLOATS == 4, "Need to update this when upgrading the VM to support >SSE2");
@@ -732,7 +888,7 @@ struct FVectorKernelExecutionIndex
 		VectorRegisterInt Index = MakeVectorRegisterInt(Context.StartInstance, Context.StartInstance + 1, Context.StartInstance + 2, Context.StartInstance + 3);
 		
 		FRegisterHandler<VectorRegisterInt> Dest(Context);
-		int32 Loops = Align(Context.NumInstances, VECTOR_WIDTH_FLOATS) / VECTOR_WIDTH_FLOATS;
+		const int32 Loops = Context.GetNumLoops<VECTOR_WIDTH_FLOATS>();
 		for (int32 i = 0; i < Loops; ++i)
 		{
 			*Dest.GetDestAndAdvance() = Index;
@@ -743,6 +899,16 @@ struct FVectorKernelExecutionIndex
 
 struct FVectorKernelEnterStatScope
 {
+	static void Optimize(FVectorVMCodeOptimizerContext& Context)
+	{
+#if STATS
+		Context.Write<FVectorVMExecFunction>(Exec);
+		FConstantHandler<int32>::Optimize(Context);
+#else
+		FConstantHandler<int32> ScopeIdx(Context.BaseContext);
+#endif
+	}
+
 	static VM_FORCEINLINE void Exec(FVectorVMContext& Context)
 	{
 		FConstantHandler<int32> ScopeIdx(Context);
@@ -758,9 +924,15 @@ struct FVectorKernelEnterStatScope
 
 struct FVectorKernelExitStatScope
 {
+	static void Optimize(FVectorVMCodeOptimizerContext& Context)
+	{
+#if STATS
+		Context.Write<FVectorVMExecFunction>(Exec);
+#endif
+	}
+		
 	static VM_FORCEINLINE void Exec(FVectorVMContext& Context)
 	{
-
 #if STATS
 		if (GbDetailedVMScriptStats)
 		{
@@ -975,6 +1147,10 @@ namespace VectorVMNoise
 		//http://mrl.nyu.edu/~perlin/noise/
 		FScalarKernelNoise3D_iNoise::Exec(Context);
 	}
+
+	static void Optimize_Noise1D(FVectorVMCodeOptimizerContext& Context) { FScalarKernelNoise1D_iNoise::Optimize(Context); }
+	static void Optimize_Noise2D(FVectorVMCodeOptimizerContext& Context) { FScalarKernelNoise2D_iNoise::Optimize(Context); }
+	static void Optimize_Noise3D(FVectorVMCodeOptimizerContext& Context) { FScalarKernelNoise3D_iNoise::Optimize(Context); }
 };
 
 //Olaf's orginal curl noise. Needs updating for the new scalar VM and possibly calling Curl Noise to avoid confusion with regular noise?
@@ -1031,18 +1207,26 @@ VectorRegister FVectorKernelNoise::RandomTable[17][17][17];
 /** Special kernel for acquiring a new ID. TODO. Can be written as general RWBuffer ops when we support that. */
 struct FScalarKernelAcquireID
 {
+	static void Optimize(FVectorVMCodeOptimizerContext& Context)
+	{
+		Context.Write<FVectorVMExecFunction>(Exec);
+		Context.Write(Context.DecodeU16());		// DataSetIndex
+		Context.Write(Context.DecodeU16());		// IDIndexReg
+		Context.Write(Context.DecodeU16());		// IDTagReg
+	}
+
 	static VM_FORCEINLINE void Exec(FVectorVMContext& Context)
 	{
-		int32 DataSetIndex = VectorVM::DecodeU16(Context);
-		TArrayView<FDataSetMeta> MetaTable = Context.DataSetMetaTable;
+		const int32 DataSetIndex = Context.DecodeU16();
+		const TArrayView<FDataSetMeta> MetaTable = Context.DataSetMetaTable;
 		TArray<int32>&RESTRICT FreeIDTable = *MetaTable[DataSetIndex].FreeIDTable;
 
-		int32 Tag = MetaTable[DataSetIndex].IDAcquireTag;
+		const int32 Tag = MetaTable[DataSetIndex].IDAcquireTag;
 
-		int32 IDIndexReg = VectorVM::DecodeU16(Context);
+		const int32 IDIndexReg = Context.DecodeU16();
 		int32*RESTRICT IDIndex = (int32*)(Context.GetTempRegister(IDIndexReg));
 
-		int32 IDTagReg = VectorVM::DecodeU16(Context);
+		const int32 IDTagReg = Context.DecodeU16();
 		int32*RESTRICT IDTag = (int32*)(Context.GetTempRegister(IDTagReg));
 
 		int32& NumFreeIDs = *MetaTable[DataSetIndex].NumFreeIDs;
@@ -1077,19 +1261,27 @@ struct FScalarKernelAcquireID
 /** Special kernel for updating a new ID. TODO. Can be written as general RWBuffer ops when we support that. */
 struct FScalarKernelUpdateID
 {
+	static void Optimize(FVectorVMCodeOptimizerContext& Context)
+	{
+		Context.Write<FVectorVMExecFunction>(Exec);
+		Context.Write(Context.DecodeU16());		// DataSetIndex
+		Context.Write(Context.DecodeU16());		// InstanceIDRegisterIndex
+		Context.Write(Context.DecodeU16());		// InstanceIndexRegisterIndex
+	}
+
 	static VM_FORCEINLINE void Exec(FVectorVMContext& Context)
 	{
-		int32 DataSetIndex = VectorVM::DecodeU16(Context);
-		int32 InstanceIDRegisterIndex = VectorVM::DecodeU16(Context);
-		int32 InstanceIndexRegisterIndex = VectorVM::DecodeU16(Context);
+		const int32 DataSetIndex = Context.DecodeU16();
+		const int32 InstanceIDRegisterIndex = Context.DecodeU16();
+		const int32 InstanceIndexRegisterIndex = Context.DecodeU16();
 
-		TArrayView<FDataSetMeta> MetaTable = Context.DataSetMetaTable;
+		const TArrayView<FDataSetMeta> MetaTable = Context.DataSetMetaTable;
 
 		TArray<int32>&RESTRICT IDTable = *MetaTable[DataSetIndex].IDTable;
-		int32 InstanceOffset = MetaTable[DataSetIndex].InstanceOffset + Context.StartInstance;
+		const int32 InstanceOffset = MetaTable[DataSetIndex].InstanceOffset + Context.StartInstance;
 
-		int32*RESTRICT IDRegister = (int32*)(Context.GetTempRegister(InstanceIDRegisterIndex));
-		int32*RESTRICT IndexRegister = (int32*)(Context.GetTempRegister(InstanceIndexRegisterIndex));
+		const int32*RESTRICT IDRegister = (int32*)(Context.GetTempRegister(InstanceIDRegisterIndex));
+		const int32*RESTRICT IndexRegister = (int32*)(Context.GetTempRegister(InstanceIndexRegisterIndex));
 		
 		FDataSetThreadLocalTempData& DataSetTempData = Context.ThreadLocalTempData[DataSetIndex];
 
@@ -1125,17 +1317,25 @@ struct FScalarKernelUpdateID
 template<typename T>
 struct FVectorKernelReadInput
 {
+	static void Optimize(FVectorVMCodeOptimizerContext& Context)
+	{
+		Context.Write<FVectorVMExecFunction>(Exec);
+		Context.Write(Context.DecodeU16());	// DataSetIndex
+		Context.Write(Context.DecodeU16());	// InputRegisterIdx
+		Context.Write(Context.DecodeU16());	// DestRegisterIdx
+	}
+
 	static VM_FORCEINLINE void Exec(FVectorVMContext& Context)
 	{
 		static const int32 InstancesPerVector = sizeof(VectorRegister) / sizeof(T);
 
-		int32 DataSetIndex = VectorVM::DecodeU16(Context);
-		int32 InputRegisterIdx = VectorVM::DecodeU16(Context);
-		int32 DestRegisterIdx = VectorVM::DecodeU16(Context);
-		int32 Loops = Align(Context.NumInstances, InstancesPerVector) / InstancesPerVector;
+		const int32 DataSetIndex = Context.DecodeU16();
+		const int32 InputRegisterIdx = Context.DecodeU16();
+		const int32 DestRegisterIdx = Context.DecodeU16();
+		const int32 Loops = Context.GetNumLoops<InstancesPerVector>();
 
 		VectorRegister* DestReg = (VectorRegister*)(Context.GetTempRegister(DestRegisterIdx));
-		VectorRegister* InputReg = (VectorRegister*)(Context.GetInputRegister<T>(DataSetIndex, InputRegisterIdx) + Context.StartInstance);
+		VectorRegister* InputReg = (VectorRegister*)(Context.GetInputRegister<T>(DataSetIndex, InputRegisterIdx) + Context.GetStartInstance());
 
 		//TODO: We can actually do some scalar loads into the first and final vectors to get around alignment issues and then use the aligned load for all others.
 		for (int32 i = 0; i < Loops; ++i)
@@ -1156,14 +1356,22 @@ struct FVectorKernelReadInput
 template<typename T>
 struct FVectorKernelReadInputNoAdvance
 {
+	static void Optimize(FVectorVMCodeOptimizerContext& Context)
+	{
+		Context.Write<FVectorVMExecFunction>(Exec);
+		Context.Write(Context.DecodeU16());	// DataSetIndex
+		Context.Write(Context.DecodeU16());	// InputRegisterIdx
+		Context.Write(Context.DecodeU16());	// DestRegisterIdx
+	}
+
 	static VM_FORCEINLINE void Exec(FVectorVMContext& Context)
 	{
 		static const int32 InstancesPerVector = sizeof(VectorRegister) / sizeof(T);
 
-		int32 DataSetIndex = VectorVM::DecodeU16(Context);
-		int32 InputRegisterIdx = VectorVM::DecodeU16(Context);
-		int32 DestRegisterIdx = VectorVM::DecodeU16(Context);
-		int32 Loops = Align(Context.NumInstances, InstancesPerVector) / InstancesPerVector;
+		const int32 DataSetIndex = Context.DecodeU16();
+		const int32 InputRegisterIdx = Context.DecodeU16();
+		const int32 DestRegisterIdx = Context.DecodeU16();
+		const int32 Loops = Context.GetNumLoops<InstancesPerVector>();
 
 		VectorRegister* DestReg = (VectorRegister*)(Context.GetTempRegister(DestRegisterIdx));
 		VectorRegister* InputReg = (VectorRegister*)(Context.GetInputRegister<T>(DataSetIndex, InputRegisterIdx));
@@ -1207,9 +1415,25 @@ struct FVectorKernelReadInputNoAdvance
 template<typename T>
 struct FScalarKernelWriteOutputIndexed
 {
+	static VM_FORCEINLINE void Optimize(FVectorVMCodeOptimizerContext& Context)
+	{
+		const uint32 SrcOpTypes = Context.BaseContext.DecodeSrcOperandTypes();
+		switch (SrcOpTypes)
+		{
+			case SRCOP_RRR: Context.Write<FVectorVMExecFunction>(DoKernel<FRegisterHandler<T>>); break;
+			case SRCOP_RRC:	Context.Write<FVectorVMExecFunction>(DoKernel<FConstantHandler<T>>); break;
+			default: check(0); break;
+		};
+
+		Context.Write(Context.DecodeU16());		// DataSetIndex
+		Context.Write(Context.DecodeU16());		// DestIndexRegisterIdx
+		Context.Write(Context.DecodeU16());		// DataHandlerType
+		Context.Write(Context.DecodeU16());		// DestRegisterIdx
+	}
+
 	static VM_FORCEINLINE void Exec(FVectorVMContext& Context)
 	{
-		uint32 SrcOpTypes = VectorVM::DecodeSrcOperandTypes(Context);
+		const uint32 SrcOpTypes = Context.DecodeSrcOperandTypes();
 		switch (SrcOpTypes)
 		{
 		case SRCOP_RRR: DoKernel<FRegisterHandler<T>>(Context); break;
@@ -1221,17 +1445,17 @@ struct FScalarKernelWriteOutputIndexed
 	template<typename DataHandlerType>
 	static VM_FORCEINLINE void DoKernel(FVectorVMContext& Context)
 	{
-		int32 DataSetIndex = VectorVM::DecodeU16(Context);
+		const int32 DataSetIndex = Context.DecodeU16();
 
-		int32 DestIndexRegisterIdx = VectorVM::DecodeU16(Context);
+		const int32 DestIndexRegisterIdx = Context.DecodeU16();
 		T* DestIndexReg = (T*)(Context.GetTempRegister(DestIndexRegisterIdx));
 
 		DataHandlerType DataHandler(Context);
 
-		int32 DestRegisterIdx = VectorVM::DecodeU16(Context);
+		const int32 DestRegisterIdx = Context.DecodeU16();
 		T* DestReg = Context.GetOutputRegister<T>(DataSetIndex, DestRegisterIdx);
 
-		for (int32 i = 0; i < Context.NumInstances; ++i)
+		for (int32 i = 0; i < Context.GetNumInstances(); ++i)
 		{
 			int32 DestIndex = *DestIndexReg;
 			if (DestIndex != INDEX_NONE)
@@ -1250,47 +1474,109 @@ struct FDataSetCounterHandler
 {
 	int32* Counter;
 	FDataSetCounterHandler(FVectorVMContext& Context)
-		: Counter(&Context.GetDataSetMeta(VectorVM::DecodeU16(Context)).DataSetAccessIndex)
+		: Counter(&Context.GetDataSetMeta(Context.DecodeU16()).DataSetAccessIndex)
 	{}
 
 	VM_FORCEINLINE void Advance() { }
 	VM_FORCEINLINE int32* Get() { return Counter; }
 	VM_FORCEINLINE int32* GetAndAdvance() { return Counter; }
 	//VM_FORCEINLINE const int32* GetDest() { return Counter; }Should never use as a dest. All kernels with read and write to this.
+
+	static void Optimize(FVectorVMCodeOptimizerContext& Context)
+	{
+		Context.Write(Context.DecodeU16());
+	}
 };
 
-template<bool bThreadsafe>
 struct FScalarKernelAcquireCounterIndex
 {
-	static VM_FORCEINLINE void Exec(FVectorVMContext& Context)
+	template<bool bThreadsafe>
+	struct InternalKernel
 	{
-		uint32 SrcOpTypes = VectorVM::DecodeSrcOperandTypes(Context);
-		switch (SrcOpTypes)
+		static VM_FORCEINLINE void DoKernel(FVectorVMContext& Context, int32* RESTRICT Dst, int32* Index, int32 Valid)
 		{
-		case SRCOP_RRR: TBinaryKernelHandler<FScalarKernelAcquireCounterIndex, FRegisterHandler<int32>, FDataSetCounterHandler, FRegisterHandler<int32>, 1>::Exec(Context); break;
-		case SRCOP_RRC:	TBinaryKernelHandler<FScalarKernelAcquireCounterIndex, FRegisterHandler<int32>, FDataSetCounterHandler, FConstantHandler<int32>, 1>::Exec(Context); break;
-		default: check(0); break;
-		};
-	}
+			if (Valid != 0)
+			{
+				*Dst = bThreadsafe ? FPlatformAtomics::InterlockedIncrement(Index) : ++(*Index);
+			}
+			else
+			{
+				*Dst = INDEX_NONE;	// Subsequent DoKernal calls above will skip over INDEX_NONE register entries...
+			}
+		}
+
+		static VM_FORCEINLINE void Exec(FVectorVMContext& Context)
+		{
+			const uint32 SrcOpType = Context.DecodeSrcOperandTypes();
+			switch (SrcOpType)
+			{
+				case SRCOP_RRR: TBinaryKernelHandler<InternalKernel<true>, FRegisterHandler<int32>, FDataSetCounterHandler, FRegisterHandler<int32>, 1>::Exec(Context); break;
+				case SRCOP_RRC:	TBinaryKernelHandler<InternalKernel<true>, FRegisterHandler<int32>, FDataSetCounterHandler, FConstantHandler<int32>, 1>::Exec(Context); break;
+				default: check(0); break;
+			};
+		}
+	};
 
-	static VM_FORCEINLINE void DoKernel(FVectorVMContext& Context, int32* RESTRICT Dst, int32* Index, int32 Valid)
+	template<bool SrcOpType>
+	static void ExecOptimized(FVectorVMContext& Context)
 	{
-		if (Valid != 0)
+		if (Context.IsParallelExecution())
 		{
-			if (bThreadsafe)
+			if (SrcOpType == SRCOP_RRR)
 			{
-				*Dst = FPlatformAtomics::InterlockedIncrement(Index);
+				TBinaryKernelHandler<InternalKernel<true>, FRegisterHandler<int32>, FDataSetCounterHandler, FRegisterHandler<int32>, 1>::Exec(Context);
 			}
 			else
 			{
-				*Dst = ++(*Index);
+				check(SrcOpType == SRCOP_RRC);
+				TBinaryKernelHandler<InternalKernel<true>, FRegisterHandler<int32>, FDataSetCounterHandler, FConstantHandler<int32>, 1>::Exec(Context);
 			}
 		}
 		else
-		{	
-			*Dst = INDEX_NONE;	// Subsequent DoKernal calls above will skip over INDEX_NONE register entries...
+		{
+			if (SrcOpType == SRCOP_RRR)
+			{
+				TBinaryKernelHandler<InternalKernel<false>, FRegisterHandler<int32>, FDataSetCounterHandler, FRegisterHandler<int32>, 1>::Exec(Context);
+			}
+			else
+			{
+				check(SrcOpType == SRCOP_RRC);
+				TBinaryKernelHandler<InternalKernel<false>, FRegisterHandler<int32>, FDataSetCounterHandler, FConstantHandler<int32>, 1>::Exec(Context);
+			}
+		}
+	}
+
+	static void Optimize(FVectorVMCodeOptimizerContext& Context)
+	{
+		const uint32 SrcOpType = Context.BaseContext.DecodeSrcOperandTypes();
+		switch (SrcOpType)
+		{
+			case SRCOP_RRR:
+				Context.Write<FVectorVMExecFunction>(FScalarKernelAcquireCounterIndex::ExecOptimized<SRCOP_RRR>);
+				break;
+			case SRCOP_RRC:
+				Context.Write<FVectorVMExecFunction>(FScalarKernelAcquireCounterIndex::ExecOptimized<SRCOP_RRC>);
+				break;
+			default: check(0); break;
+		}
+		// Three registers, note we don't call Optimize on the Kernel since that will write the Exec and we are selecting based upon thread safe or not
+		Context.Write(Context.DecodeU16());
+		Context.Write(Context.DecodeU16());
+		Context.Write(Context.DecodeU16());
+	}
+
+	static VM_FORCEINLINE void Exec(FVectorVMContext& Context)
+	{
+		if ( Context.IsParallelExecution() )
+		{
+			InternalKernel<true>::Exec(Context);
+		}
+		else
+		{
+			InternalKernel<false>::Exec(Context);
 		}
 	}
+
 };
 
 //TODO: REWORK TO FUNCITON LIKE THE ABOVE.
@@ -1317,9 +1603,23 @@ struct FScalarKernelAcquireCounterIndex
 
 struct FKernelExternalFunctionCall
 {
+	static void Optimize(FVectorVMCodeOptimizerContext& Context)
+	{
+		const uint32 ExternalFuncIdx = Context.DecodeU8();
+
+		Context.Write<FVectorVMExecFunction>(Exec);
+		Context.Write<uint8>(ExternalFuncIdx);
+
+		const int32 NumRegisters = Context.ExternalFunctionRegisterCounts[ExternalFuncIdx];
+		for ( int32 i=0; i < NumRegisters; ++i )
+		{
+			Context.Write(Context.DecodeU16());
+		}
+	}
+
 	static void Exec(FVectorVMContext& Context)
 	{
-		uint32 ExternalFuncIdx = VectorVM::DecodeU8(Context);
+		const uint32 ExternalFuncIdx = Context.DecodeU8();
 		Context.ExternalFunctionTable[ExternalFuncIdx].Execute(Context);
 	}
 };
@@ -1761,7 +2061,8 @@ void VectorVM::Init()
 }
 
 void VectorVM::Exec(
-	uint8 const* Code,
+	uint8 const* ByteCode,
+	uint8 const* OptimizedByteCode,
 	int32 NumTempRegisters,
 	uint8 const* ConstantTable,
 	TArrayView<FDataSetMeta> DataSetMetaTable,
@@ -1776,18 +2077,19 @@ void VectorVM::Exec(
 	//TRACE_CPUPROFILER_EVENT_SCOPE("VMExec");
 	SCOPE_CYCLE_COUNTER(STAT_VVMExec);
 
-	int32 MaxInstances = FMath::Min(GParallelVVMInstancesPerChunk, NumInstances);
-	int32 NumChunks = (NumInstances / GParallelVVMInstancesPerChunk) + 1;
-	int32 ChunksPerBatch = (GbParallelVVM != 0 && FApp::ShouldUseThreadingForPerformance()) ? GParallelVVMChunksPerBatch : NumChunks;
-	int32 NumBatches = FMath::DivideAndRoundUp(NumChunks, ChunksPerBatch);
-	bool bParallel = NumBatches > 1;
+	const int32 MaxInstances = FMath::Min(GParallelVVMInstancesPerChunk, NumInstances);
+	const int32 NumChunks = (NumInstances / GParallelVVMInstancesPerChunk) + 1;
+	const int32 ChunksPerBatch = (GbParallelVVM != 0 && FApp::ShouldUseThreadingForPerformance()) ? GParallelVVMChunksPerBatch : NumChunks;
+	const int32 NumBatches = FMath::DivideAndRoundUp(NumChunks, ChunksPerBatch);
+	const bool bParallel = NumBatches > 1;
+	const bool bUseOptimizedByteCode = (OptimizedByteCode != nullptr) && GbUseOptimizedVMByteCode;
 
 	auto ExecChunkBatch = [&](int32 BatchIdx)
 	{
 		//SCOPE_CYCLE_COUNTER(STAT_VVMExecChunk);
 
 		FVectorVMContext& Context = FVectorVMContext::Get();
-		Context.PrepareForExec(NumTempRegisters, ConstantTable, ExternalFunctionTable, UserPtrTable, DataSetMetaTable, MaxInstances
+		Context.PrepareForExec(NumTempRegisters, ConstantTable, ExternalFunctionTable, UserPtrTable, DataSetMetaTable, MaxInstances, bParallel
 #if STATS
 			, &StatScopes
 #endif
@@ -1795,152 +2097,165 @@ void VectorVM::Exec(
 
 		// Process one chunk at a time.
 		int32 ChunkIdx = BatchIdx * ChunksPerBatch;
-		int32 FirstInstance = ChunkIdx * GParallelVVMInstancesPerChunk;
-		int32 FinalInstance = FMath::Min(NumInstances, FirstInstance + (ChunksPerBatch * GParallelVVMInstancesPerChunk));
+		const int32 FirstInstance = ChunkIdx * GParallelVVMInstancesPerChunk;
+		const int32 FinalInstance = FMath::Min(NumInstances, FirstInstance + (ChunksPerBatch * GParallelVVMInstancesPerChunk));
 		int32 InstancesLeft = FinalInstance - FirstInstance;
 		while (InstancesLeft > 0)
 		{
 			int32 NumInstancesThisChunk = FMath::Min(InstancesLeft, (int32)GParallelVVMInstancesPerChunk);
 			int32 StartInstance = GParallelVVMInstancesPerChunk * ChunkIdx;
-			// Setup execution context.
-			Context.PrepareForChunk(Code, NumInstancesThisChunk, StartInstance);
 
-			EVectorVMOp Op = EVectorVMOp::done;
-
-			// Execute VM on all vectors in this chunk.
-			do
+			// Execute optimized byte code version
+			if ( bUseOptimizedByteCode )
 			{
-				Op = DecodeOp(Context);
-				switch (Op)
-				{
-					// Dispatch kernel ops.
-				case EVectorVMOp::add: FVectorKernelAdd::Exec(Context); break;
-				case EVectorVMOp::sub: FVectorKernelSub::Exec(Context); break;
-				case EVectorVMOp::mul: FVectorKernelMul::Exec(Context); break;
-				case EVectorVMOp::div: FVectorKernelDiv::Exec(Context); break;
-				case EVectorVMOp::mad: FVectorKernelMad::Exec(Context); break;
-				case EVectorVMOp::lerp: FVectorKernelLerp::Exec(Context); break;
-				case EVectorVMOp::rcp: FVectorKernelRcp::Exec(Context); break;
-				case EVectorVMOp::rsq: FVectorKernelRsq::Exec(Context); break;
-				case EVectorVMOp::sqrt: FVectorKernelSqrt::Exec(Context); break;
-				case EVectorVMOp::neg: FVectorKernelNeg::Exec(Context); break;
-				case EVectorVMOp::abs: FVectorKernelAbs::Exec(Context); break;
-				case EVectorVMOp::exp: FVectorKernelExp::Exec(Context); break;
-				case EVectorVMOp::exp2: FVectorKernelExp2::Exec(Context); break;
-				case EVectorVMOp::log: FVectorKernelLog::Exec(Context); break;
-				case EVectorVMOp::log2: FVectorKernelLog2::Exec(Context); break;
-				case EVectorVMOp::sin: FVectorKernelSin::Exec(Context); break;
-				case EVectorVMOp::cos: FVectorKernelCos::Exec(Context); break;
-				case EVectorVMOp::tan: FVectorKernelTan::Exec(Context); break;
-				case EVectorVMOp::asin: FVectorKernelASin::Exec(Context); break;
-				case EVectorVMOp::acos: FVectorKernelACos::Exec(Context); break;
-				case EVectorVMOp::atan: FVectorKernelATan::Exec(Context); break;
-				case EVectorVMOp::atan2: FVectorKernelATan2::Exec(Context); break;
-				case EVectorVMOp::ceil: FVectorKernelCeil::Exec(Context); break;
-				case EVectorVMOp::floor: FVectorKernelFloor::Exec(Context); break;
-				case EVectorVMOp::round: FVectorKernelRound::Exec(Context); break;
-				case EVectorVMOp::fmod: FVectorKernelMod::Exec(Context); break;
-				case EVectorVMOp::frac: FVectorKernelFrac::Exec(Context); break;
-				case EVectorVMOp::trunc: FVectorKernelTrunc::Exec(Context); break;
-				case EVectorVMOp::clamp: FVectorKernelClamp::Exec(Context); break;
-				case EVectorVMOp::min: FVectorKernelMin::Exec(Context); break;
-				case EVectorVMOp::max: FVectorKernelMax::Exec(Context); break;
-				case EVectorVMOp::pow: FVectorKernelPow::Exec(Context); break;
-				case EVectorVMOp::sign: FVectorKernelSign::Exec(Context); break;
-				case EVectorVMOp::step: FVectorKernelStep::Exec(Context); break;
-				case EVectorVMOp::random: FVectorKernelRandom::Exec(Context); break;
-				case EVectorVMOp::noise: VectorVMNoise::Noise1D(Context); break;
-				case EVectorVMOp::noise2D: VectorVMNoise::Noise2D(Context); break;
-				case EVectorVMOp::noise3D: VectorVMNoise::Noise3D(Context); break;
-
-				case EVectorVMOp::cmplt: FVectorKernelCompareLT::Exec(Context); break;
-				case EVectorVMOp::cmple: FVectorKernelCompareLE::Exec(Context); break;
-				case EVectorVMOp::cmpgt: FVectorKernelCompareGT::Exec(Context); break;
-				case EVectorVMOp::cmpge: FVectorKernelCompareGE::Exec(Context); break;
-				case EVectorVMOp::cmpeq: FVectorKernelCompareEQ::Exec(Context); break;
-				case EVectorVMOp::cmpneq: FVectorKernelCompareNEQ::Exec(Context); break;
-				case EVectorVMOp::select: FVectorKernelSelect::Exec(Context); break;
-
-				case EVectorVMOp::addi: FVectorIntKernelAdd::Exec(Context); break;
-				case EVectorVMOp::subi: FVectorIntKernelSubtract::Exec(Context); break;
-				case EVectorVMOp::muli: FVectorIntKernelMultiply::Exec(Context); break;
-				case EVectorVMOp::divi: FVectorIntKernelDivide::Exec(Context); break;
-				case EVectorVMOp::clampi: FVectorIntKernelClamp::Exec(Context); break;
-				case EVectorVMOp::mini: FVectorIntKernelMin::Exec(Context); break;
-				case EVectorVMOp::maxi: FVectorIntKernelMax::Exec(Context); break;
-				case EVectorVMOp::absi: FVectorIntKernelAbs::Exec(Context); break;
-				case EVectorVMOp::negi: FVectorIntKernelNegate::Exec(Context); break;
-				case EVectorVMOp::signi: FVectorIntKernelSign::Exec(Context); break;
-				case EVectorVMOp::randomi: FScalarIntKernelRandom::Exec(Context); break;
-				case EVectorVMOp::cmplti: FVectorIntKernelCompareLT::Exec(Context); break;
-				case EVectorVMOp::cmplei: FVectorIntKernelCompareLE::Exec(Context); break;
-				case EVectorVMOp::cmpgti: FVectorIntKernelCompareGT::Exec(Context); break;
-				case EVectorVMOp::cmpgei: FVectorIntKernelCompareGE::Exec(Context); break;
-				case EVectorVMOp::cmpeqi: FVectorIntKernelCompareEQ::Exec(Context); break;
-				case EVectorVMOp::cmpneqi: FVectorIntKernelCompareNEQ::Exec(Context); break;
-				case EVectorVMOp::bit_and: FVectorIntKernelBitAnd::Exec(Context); break;
-				case EVectorVMOp::bit_or: FVectorIntKernelBitOr::Exec(Context); break;
-				case EVectorVMOp::bit_xor: FVectorIntKernelBitXor::Exec(Context); break;
-				case EVectorVMOp::bit_not: FVectorIntKernelBitNot::Exec(Context); break;
-				case EVectorVMOp::bit_lshift: FVectorIntKernelBitLShift::Exec(Context); break;
-				case EVectorVMOp::bit_rshift: FVectorIntKernelBitRShift::Exec(Context); break;
-				case EVectorVMOp::logic_and: FVectorIntKernelLogicAnd::Exec(Context); break;
-				case EVectorVMOp::logic_or: FVectorIntKernelLogicOr::Exec(Context); break;
-				case EVectorVMOp::logic_xor: FVectorIntKernelLogicXor::Exec(Context); break;
-				case EVectorVMOp::logic_not: FVectorIntKernelLogicNot::Exec(Context); break;
-				case EVectorVMOp::f2i: FVectorKernelFloatToInt::Exec(Context); break;
-				case EVectorVMOp::i2f: FVectorKernelIntToFloat::Exec(Context); break;
-				case EVectorVMOp::f2b: FVectorKernelFloatToBool::Exec(Context); break;
-				case EVectorVMOp::b2f: FVectorKernelBoolToFloat::Exec(Context); break;
-				case EVectorVMOp::i2b: FVectorKernelIntToBool::Exec(Context); break;
-				case EVectorVMOp::b2i: FVectorKernelBoolToInt::Exec(Context); break;
-
-				case EVectorVMOp::outputdata_32bit:	FScalarKernelWriteOutputIndexed<int32>::Exec(Context);	break;
-				case EVectorVMOp::inputdata_32bit: FVectorKernelReadInput<int32>::Exec(Context); break;
-					//case EVectorVMOp::inputdata_32bit: FVectorKernelReadInput32::Exec(Context); break;
-				case EVectorVMOp::inputdata_noadvance_32bit: FVectorKernelReadInputNoAdvance<int32>::Exec(Context); break;
-				case EVectorVMOp::acquireindex:	
+				// Setup execution context.
+				Context.PrepareForChunk(OptimizedByteCode, NumInstancesThisChunk, StartInstance);
+
+				while (true)
 				{
-					if (bParallel)
+					FVectorVMExecFunction ExecFunction = reinterpret_cast<FVectorVMExecFunction>(Context.DecodePtr());
+					if (ExecFunction == nullptr)
 					{
-						FScalarKernelAcquireCounterIndex<true>::Exec(Context);
+						break;
 					}
-					else
+					ExecFunction(Context);
+				}
+			}
+			else
+			{
+				// Setup execution context.
+				Context.PrepareForChunk(ByteCode, NumInstancesThisChunk, StartInstance);
+
+				// Execute VM on all vectors in this chunk.
+				EVectorVMOp Op = EVectorVMOp::done;
+				do
+				{
+					Op = Context.DecodeOp();
+					switch (Op)
 					{
-						FScalarKernelAcquireCounterIndex<false>::Exec(Context);
+						// Dispatch kernel ops.
+						case EVectorVMOp::add: FVectorKernelAdd::Exec(Context); break;
+						case EVectorVMOp::sub: FVectorKernelSub::Exec(Context); break;
+						case EVectorVMOp::mul: FVectorKernelMul::Exec(Context); break;
+						case EVectorVMOp::div: FVectorKernelDiv::Exec(Context); break;
+						case EVectorVMOp::mad: FVectorKernelMad::Exec(Context); break;
+						case EVectorVMOp::lerp: FVectorKernelLerp::Exec(Context); break;
+						case EVectorVMOp::rcp: FVectorKernelRcp::Exec(Context); break;
+						case EVectorVMOp::rsq: FVectorKernelRsq::Exec(Context); break;
+						case EVectorVMOp::sqrt: FVectorKernelSqrt::Exec(Context); break;
+						case EVectorVMOp::neg: FVectorKernelNeg::Exec(Context); break;
+						case EVectorVMOp::abs: FVectorKernelAbs::Exec(Context); break;
+						case EVectorVMOp::exp: FVectorKernelExp::Exec(Context); break;
+						case EVectorVMOp::exp2: FVectorKernelExp2::Exec(Context); break;
+						case EVectorVMOp::log: FVectorKernelLog::Exec(Context); break;
+						case EVectorVMOp::log2: FVectorKernelLog2::Exec(Context); break;
+						case EVectorVMOp::sin: FVectorKernelSin::Exec(Context); break;
+						case EVectorVMOp::cos: FVectorKernelCos::Exec(Context); break;
+						case EVectorVMOp::tan: FVectorKernelTan::Exec(Context); break;
+						case EVectorVMOp::asin: FVectorKernelASin::Exec(Context); break;
+						case EVectorVMOp::acos: FVectorKernelACos::Exec(Context); break;
+						case EVectorVMOp::atan: FVectorKernelATan::Exec(Context); break;
+						case EVectorVMOp::atan2: FVectorKernelATan2::Exec(Context); break;
+						case EVectorVMOp::ceil: FVectorKernelCeil::Exec(Context); break;
+						case EVectorVMOp::floor: FVectorKernelFloor::Exec(Context); break;
+						case EVectorVMOp::round: FVectorKernelRound::Exec(Context); break;
+						case EVectorVMOp::fmod: FVectorKernelMod::Exec(Context); break;
+						case EVectorVMOp::frac: FVectorKernelFrac::Exec(Context); break;
+						case EVectorVMOp::trunc: FVectorKernelTrunc::Exec(Context); break;
+						case EVectorVMOp::clamp: FVectorKernelClamp::Exec(Context); break;
+						case EVectorVMOp::min: FVectorKernelMin::Exec(Context); break;
+						case EVectorVMOp::max: FVectorKernelMax::Exec(Context); break;
+						case EVectorVMOp::pow: FVectorKernelPow::Exec(Context); break;
+						case EVectorVMOp::sign: FVectorKernelSign::Exec(Context); break;
+						case EVectorVMOp::step: FVectorKernelStep::Exec(Context); break;
+						case EVectorVMOp::random: FVectorKernelRandom::Exec(Context); break;
+						case EVectorVMOp::noise: VectorVMNoise::Noise1D(Context); break;
+						case EVectorVMOp::noise2D: VectorVMNoise::Noise2D(Context); break;
+						case EVectorVMOp::noise3D: VectorVMNoise::Noise3D(Context); break;
+
+						case EVectorVMOp::cmplt: FVectorKernelCompareLT::Exec(Context); break;
+						case EVectorVMOp::cmple: FVectorKernelCompareLE::Exec(Context); break;
+						case EVectorVMOp::cmpgt: FVectorKernelCompareGT::Exec(Context); break;
+						case EVectorVMOp::cmpge: FVectorKernelCompareGE::Exec(Context); break;
+						case EVectorVMOp::cmpeq: FVectorKernelCompareEQ::Exec(Context); break;
+						case EVectorVMOp::cmpneq: FVectorKernelCompareNEQ::Exec(Context); break;
+						case EVectorVMOp::select: FVectorKernelSelect::Exec(Context); break;
+
+						case EVectorVMOp::addi: FVectorIntKernelAdd::Exec(Context); break;
+						case EVectorVMOp::subi: FVectorIntKernelSubtract::Exec(Context); break;
+						case EVectorVMOp::muli: FVectorIntKernelMultiply::Exec(Context); break;
+						case EVectorVMOp::divi: FVectorIntKernelDivide::Exec(Context); break;
+						case EVectorVMOp::clampi: FVectorIntKernelClamp::Exec(Context); break;
+						case EVectorVMOp::mini: FVectorIntKernelMin::Exec(Context); break;
+						case EVectorVMOp::maxi: FVectorIntKernelMax::Exec(Context); break;
+						case EVectorVMOp::absi: FVectorIntKernelAbs::Exec(Context); break;
+						case EVectorVMOp::negi: FVectorIntKernelNegate::Exec(Context); break;
+						case EVectorVMOp::signi: FVectorIntKernelSign::Exec(Context); break;
+						case EVectorVMOp::randomi: FScalarIntKernelRandom::Exec(Context); break;
+						case EVectorVMOp::cmplti: FVectorIntKernelCompareLT::Exec(Context); break;
+						case EVectorVMOp::cmplei: FVectorIntKernelCompareLE::Exec(Context); break;
+						case EVectorVMOp::cmpgti: FVectorIntKernelCompareGT::Exec(Context); break;
+						case EVectorVMOp::cmpgei: FVectorIntKernelCompareGE::Exec(Context); break;
+						case EVectorVMOp::cmpeqi: FVectorIntKernelCompareEQ::Exec(Context); break;
+						case EVectorVMOp::cmpneqi: FVectorIntKernelCompareNEQ::Exec(Context); break;
+						case EVectorVMOp::bit_and: FVectorIntKernelBitAnd::Exec(Context); break;
+						case EVectorVMOp::bit_or: FVectorIntKernelBitOr::Exec(Context); break;
+						case EVectorVMOp::bit_xor: FVectorIntKernelBitXor::Exec(Context); break;
+						case EVectorVMOp::bit_not: FVectorIntKernelBitNot::Exec(Context); break;
+						case EVectorVMOp::bit_lshift: FVectorIntKernelBitLShift::Exec(Context); break;
+						case EVectorVMOp::bit_rshift: FVectorIntKernelBitRShift::Exec(Context); break;
+						case EVectorVMOp::logic_and: FVectorIntKernelLogicAnd::Exec(Context); break;
+						case EVectorVMOp::logic_or: FVectorIntKernelLogicOr::Exec(Context); break;
+						case EVectorVMOp::logic_xor: FVectorIntKernelLogicXor::Exec(Context); break;
+						case EVectorVMOp::logic_not: FVectorIntKernelLogicNot::Exec(Context); break;
+						case EVectorVMOp::f2i: FVectorKernelFloatToInt::Exec(Context); break;
+						case EVectorVMOp::i2f: FVectorKernelIntToFloat::Exec(Context); break;
+						case EVectorVMOp::f2b: FVectorKernelFloatToBool::Exec(Context); break;
+						case EVectorVMOp::b2f: FVectorKernelBoolToFloat::Exec(Context); break;
+						case EVectorVMOp::i2b: FVectorKernelIntToBool::Exec(Context); break;
+						case EVectorVMOp::b2i: FVectorKernelBoolToInt::Exec(Context); break;
+
+						case EVectorVMOp::outputdata_32bit:	FScalarKernelWriteOutputIndexed<int32>::Exec(Context);	break;
+						case EVectorVMOp::inputdata_32bit: FVectorKernelReadInput<int32>::Exec(Context); break;
+						case EVectorVMOp::inputdata_noadvance_32bit: FVectorKernelReadInputNoAdvance<int32>::Exec(Context); break;
+						case EVectorVMOp::acquireindex:	FScalarKernelAcquireCounterIndex::Exec(Context); break;
+						case EVectorVMOp::external_func_call: FKernelExternalFunctionCall::Exec(Context); break;
+
+						case EVectorVMOp::exec_index: FVectorKernelExecutionIndex::Exec(Context); break;
+
+						case EVectorVMOp::enter_stat_scope: FVectorKernelEnterStatScope::Exec(Context); break;
+						case EVectorVMOp::exit_stat_scope: FVectorKernelExitStatScope::Exec(Context); break;
+
+						//Special case ops to handle unique IDs but this can be written as generalized buffer operations. TODO!
+						case EVectorVMOp::update_id:	FScalarKernelUpdateID::Exec(Context); break;
+						case EVectorVMOp::acquire_id:	FScalarKernelAcquireID::Exec(Context); break;
+
+						// Execution always terminates with a "done" opcode.
+						case EVectorVMOp::done:
+							break;
+
+						// Opcode not recognized / implemented.
+						default:
+							UE_LOG(LogVectorVM, Fatal, TEXT("Unknown op code 0x%02x"), (uint32)Op);
+							return;//BAIL
 					}
-
-				}break;
-				case EVectorVMOp::external_func_call: FKernelExternalFunctionCall::Exec(Context); break;
-
-				case EVectorVMOp::exec_index: FVectorKernelExecutionIndex::Exec(Context); break;
-
-				case EVectorVMOp::enter_stat_scope: FVectorKernelEnterStatScope::Exec(Context); break;
-				case EVectorVMOp::exit_stat_scope: FVectorKernelExitStatScope::Exec(Context); break;
-
-			//Special case ops to handle unique IDs but this can be written as generalized buffer operations. TODO!
-			case EVectorVMOp::update_id:	FScalarKernelUpdateID::Exec(Context); break;
-			case EVectorVMOp::acquire_id:	FScalarKernelAcquireID::Exec(Context); break;
-
-					// Execution always terminates with a "done" opcode.
-				case EVectorVMOp::done:
-					break;
-
-					// Opcode not recognized / implemented.
-				default:
-					UE_LOG(LogVectorVM, Fatal, TEXT("Unknown op code 0x%02x"), (uint32)Op);
-					return;//BAIL
-				}
-			} while (Op != EVectorVMOp::done);
+				} while (Op != EVectorVMOp::done);
+			}
 
 			InstancesLeft -= GParallelVVMInstancesPerChunk;
 			++ChunkIdx;
 		}
-
 		Context.FinishExec();
 	};
 
-	ParallelFor(NumBatches, ExecChunkBatch, GbParallelVVM == 0 || !bParallel);
+	if ( NumBatches > 1 )
+	{
+		ParallelFor(NumBatches, ExecChunkBatch, GbParallelVVM == 0 || !bParallel);
+	}
+	else
+	{
+		ExecChunkBatch(0);
+	}
 }
 
 uint8 VectorVM::GetNumOpCodes()
@@ -1970,4 +2285,137 @@ FString VectorVM::GetOperandLocationName(EVectorVMOperandLocation Location)
 }
 #endif
 
+void VectorVM::OptimizeByteCode(const uint8* ByteCode, TArray<uint8>& OptimizedCode, TArrayView<uint8> ExternalFunctionRegisterCounts)
+{
+	OptimizedCode.Empty();
+	if (!PLATFORM_SUPPORTS_UNALIGNED_LOADS || !PLATFORM_LITTLE_ENDIAN)
+	{
+		//-TODO: Support unaligned writes & little endian
+		return;
+	}
+
+	if ( !GbOptimizeVMByteCode || (ByteCode == nullptr) )
+	{
+		return;
+	}
+
+	FVectorVMCodeOptimizerContext Context(FVectorVMContext::Get(), ByteCode, OptimizedCode, ExternalFunctionRegisterCounts);
+
+	EVectorVMOp Op = EVectorVMOp::done;
+	do
+	{
+		Op = Context.BaseContext.DecodeOp();
+
+		switch (Op)
+		{
+			case EVectorVMOp::add: FVectorKernelAdd::Optimize(Context); break;
+			case EVectorVMOp::sub: FVectorKernelSub::Optimize(Context); break;
+			case EVectorVMOp::mul: FVectorKernelMul::Optimize(Context); break;
+			case EVectorVMOp::div: FVectorKernelDiv::Optimize(Context); break;
+			case EVectorVMOp::mad: FVectorKernelMad::Optimize(Context); break;
+			case EVectorVMOp::lerp: FVectorKernelLerp::Optimize(Context); break;
+			case EVectorVMOp::rcp: FVectorKernelRcp::Optimize(Context); break;
+			case EVectorVMOp::rsq: FVectorKernelRsq::Optimize(Context); break;
+			case EVectorVMOp::sqrt: FVectorKernelSqrt::Optimize(Context); break;
+			case EVectorVMOp::neg: FVectorKernelNeg::Optimize(Context); break;
+			case EVectorVMOp::abs: FVectorKernelAbs::Optimize(Context); break;
+			case EVectorVMOp::exp: FVectorKernelExp::Optimize(Context); break;
+			case EVectorVMOp::exp2: FVectorKernelExp2::Optimize(Context); break;
+			case EVectorVMOp::log: FVectorKernelLog::Optimize(Context); break;
+			case EVectorVMOp::log2: FVectorKernelLog2::Optimize(Context); break;
+			case EVectorVMOp::sin: FVectorKernelSin::Optimize(Context); break;
+			case EVectorVMOp::cos: FVectorKernelCos::Optimize(Context); break;
+			case EVectorVMOp::tan: FVectorKernelTan::Optimize(Context); break;
+			case EVectorVMOp::asin: FVectorKernelASin::Optimize(Context); break;
+			case EVectorVMOp::acos: FVectorKernelACos::Optimize(Context); break;
+			case EVectorVMOp::atan: FVectorKernelATan::Optimize(Context); break;
+			case EVectorVMOp::atan2: FVectorKernelATan2::Optimize(Context); break;
+			case EVectorVMOp::ceil: FVectorKernelCeil::Optimize(Context); break;
+			case EVectorVMOp::floor: FVectorKernelFloor::Optimize(Context); break;
+			case EVectorVMOp::round: FVectorKernelRound::Optimize(Context); break;
+			case EVectorVMOp::fmod: FVectorKernelMod::Optimize(Context); break;
+			case EVectorVMOp::frac: FVectorKernelFrac::Optimize(Context); break;
+			case EVectorVMOp::trunc: FVectorKernelTrunc::Optimize(Context); break;
+			case EVectorVMOp::clamp: FVectorKernelClamp::Optimize(Context); break;
+			case EVectorVMOp::min: FVectorKernelMin::Optimize(Context); break;
+			case EVectorVMOp::max: FVectorKernelMax::Optimize(Context); break;
+			case EVectorVMOp::pow: FVectorKernelPow::Optimize(Context); break;
+			case EVectorVMOp::sign: FVectorKernelSign::Optimize(Context); break;
+			case EVectorVMOp::step: FVectorKernelStep::Optimize(Context); break;
+			case EVectorVMOp::random: FVectorKernelRandom::Optimize(Context); break;
+			case EVectorVMOp::noise: VectorVMNoise::Optimize_Noise1D(Context); break;
+			case EVectorVMOp::noise2D: VectorVMNoise::Optimize_Noise2D(Context); break;
+			case EVectorVMOp::noise3D: VectorVMNoise::Optimize_Noise3D(Context); break;
+
+			case EVectorVMOp::cmplt: FVectorKernelCompareLT::Optimize(Context); break;
+			case EVectorVMOp::cmple: FVectorKernelCompareLE::Optimize(Context); break;
+			case EVectorVMOp::cmpgt: FVectorKernelCompareGT::Optimize(Context); break;
+			case EVectorVMOp::cmpge: FVectorKernelCompareGE::Optimize(Context); break;
+			case EVectorVMOp::cmpeq: FVectorKernelCompareEQ::Optimize(Context); break;
+			case EVectorVMOp::cmpneq: FVectorKernelCompareNEQ::Optimize(Context); break;
+			case EVectorVMOp::select: FVectorKernelSelect::Optimize(Context); break;
+
+			case EVectorVMOp::addi: FVectorIntKernelAdd::Optimize(Context); break;
+			case EVectorVMOp::subi: FVectorIntKernelSubtract::Optimize(Context); break;
+			case EVectorVMOp::muli: FVectorIntKernelMultiply::Optimize(Context); break;
+			case EVectorVMOp::divi: FVectorIntKernelDivide::Optimize(Context); break;
+			case EVectorVMOp::clampi: FVectorIntKernelClamp::Optimize(Context); break;
+			case EVectorVMOp::mini: FVectorIntKernelMin::Optimize(Context); break;
+			case EVectorVMOp::maxi: FVectorIntKernelMax::Optimize(Context); break;
+			case EVectorVMOp::absi: FVectorIntKernelAbs::Optimize(Context); break;
+			case EVectorVMOp::negi: FVectorIntKernelNegate::Optimize(Context); break;
+			case EVectorVMOp::signi: FVectorIntKernelSign::Optimize(Context); break;
+			case EVectorVMOp::randomi: FScalarIntKernelRandom::Optimize(Context); break;
+			case EVectorVMOp::cmplti: FVectorIntKernelCompareLT::Optimize(Context); break;
+			case EVectorVMOp::cmplei: FVectorIntKernelCompareLE::Optimize(Context); break;
+			case EVectorVMOp::cmpgti: FVectorIntKernelCompareGT::Optimize(Context); break;
+			case EVectorVMOp::cmpgei: FVectorIntKernelCompareGE::Optimize(Context); break;
+			case EVectorVMOp::cmpeqi: FVectorIntKernelCompareEQ::Optimize(Context); break;
+			case EVectorVMOp::cmpneqi: FVectorIntKernelCompareNEQ::Optimize(Context); break;
+			case EVectorVMOp::bit_and: FVectorIntKernelBitAnd::Optimize(Context); break;
+			case EVectorVMOp::bit_or: FVectorIntKernelBitOr::Optimize(Context); break;
+			case EVectorVMOp::bit_xor: FVectorIntKernelBitXor::Optimize(Context); break;
+			case EVectorVMOp::bit_not: FVectorIntKernelBitNot::Optimize(Context); break;
+			case EVectorVMOp::bit_lshift: FVectorIntKernelBitLShift::Optimize(Context); break;
+			case EVectorVMOp::bit_rshift: FVectorIntKernelBitRShift::Optimize(Context); break;
+			case EVectorVMOp::logic_and: FVectorIntKernelLogicAnd::Optimize(Context); break;
+			case EVectorVMOp::logic_or: FVectorIntKernelLogicOr::Optimize(Context); break;
+			case EVectorVMOp::logic_xor: FVectorIntKernelLogicXor::Optimize(Context); break;
+			case EVectorVMOp::logic_not: FVectorIntKernelLogicNot::Optimize(Context); break;
+			case EVectorVMOp::f2i: FVectorKernelFloatToInt::Optimize(Context); break;
+			case EVectorVMOp::i2f: FVectorKernelIntToFloat::Optimize(Context); break;
+			case EVectorVMOp::f2b: FVectorKernelFloatToBool::Optimize(Context); break;
+			case EVectorVMOp::b2f: FVectorKernelBoolToFloat::Optimize(Context); break;
+			case EVectorVMOp::i2b: FVectorKernelIntToBool::Optimize(Context); break;
+			case EVectorVMOp::b2i: FVectorKernelBoolToInt::Optimize(Context); break;
+
+			case EVectorVMOp::outputdata_32bit:	FScalarKernelWriteOutputIndexed<int32>::Optimize(Context);	break;
+			case EVectorVMOp::inputdata_32bit: FVectorKernelReadInput<int32>::Optimize(Context); break;
+			case EVectorVMOp::inputdata_noadvance_32bit: FVectorKernelReadInputNoAdvance<int32>::Optimize(Context); break;
+			case EVectorVMOp::acquireindex: FScalarKernelAcquireCounterIndex::Optimize(Context); break;
+			case EVectorVMOp::external_func_call: FKernelExternalFunctionCall::Optimize(Context); break;
+
+			case EVectorVMOp::exec_index: FVectorKernelExecutionIndex::Optimize(Context); break;
+
+			case EVectorVMOp::enter_stat_scope: FVectorKernelEnterStatScope::Optimize(Context); break;
+			case EVectorVMOp::exit_stat_scope: FVectorKernelExitStatScope::Optimize(Context); break;
+
+			//Special case ops to handle unique IDs but this can be written as generalized buffer operations. TODO!
+			case EVectorVMOp::update_id:	FScalarKernelUpdateID::Optimize(Context); break;
+			case EVectorVMOp::acquire_id:	FScalarKernelAcquireID::Optimize(Context); break;
+
+			// Execution always terminates with a "done" opcode.
+			case EVectorVMOp::done:
+				break;
+
+				// Opcode not recognized / implemented.
+			default:
+				UE_LOG(LogVectorVM, Fatal, TEXT("Unknown op code 0x%02x"), (uint32)Op);
+				OptimizedCode.Empty();
+				return;//BAIL
+		}
+	} while (Op != EVectorVMOp::done);
+	Context.Write<FVectorVMExecFunction>(nullptr);
+}
+
 #undef VM_FORCEINLINE
diff --git a/Engine/Source/Runtime/VectorVM/Public/VectorVM.h b/Engine/Source/Runtime/VectorVM/Public/VectorVM.h
--- a/Engine/Source/Runtime/VectorVM/Public/VectorVM.h
+++ b/Engine/Source/Runtime/VectorVM/Public/VectorVM.h
@@ -5,6 +5,8 @@
 #include "CoreMinimal.h"
 #include "UObject/ObjectMacros.h"
 #include "Math/RandomStream.h"
+#include "Misc/ByteSwap.h"
+#include "Templates/AlignmentTemplates.h"
 
 //TODO: move to a per platform header and have VM scale vectorization according to vector width.
 #define VECTOR_WIDTH (128)
@@ -257,8 +259,10 @@ struct FDataSetThreadLocalTempData
 */
 struct FVectorVMContext : TThreadSingleton<FVectorVMContext>
 {
+private:
 	/** Pointer to the next element in the byte code. */
 	uint8 const* RESTRICT Code;
+public:
 	/** Pointer to the constant table. */
 	uint8 const* RESTRICT ConstantTable;
 	/** Num temp registers required by this script. */
@@ -268,8 +272,11 @@ struct FVectorVMContext : TThreadSingleton<FVectorVMContext>
 	FVMExternalFunction* RESTRICT ExternalFunctionTable;
 	/** Table of user pointers.*/
 	void** UserPtrTable;
+
 	/** Number of instances to process. */
 	int32 NumInstances;
+	/** Number of instances to process when doing batches of VECTOR_WIDTH_FLOATS. */
+	int32 NumInstancesVectorFloats;
 	/** Start instance of current chunk. */
 	int32 StartInstance;
 
@@ -293,6 +300,8 @@ struct FVectorVMContext : TThreadSingleton<FVectorVMContext>
 	/** Thread local per instance random counters for use in external functions needing deterministic randoms. */
 	TArray<int32> RandCounters;
 
+	bool bIsParallelExecution;
+
 	FVectorVMContext();
 
 	void PrepareForExec(
@@ -301,7 +310,8 @@ struct FVectorVMContext : TThreadSingleton<FVectorVMContext>
 		FVMExternalFunction* InExternalFunctionTable,
 		void** InUserPtrTable,
 		TArrayView<FDataSetMeta> InDataSetMetaTable,
-		int32 MaxNumInstances
+		int32 MaxNumInstances,
+		bool bInParallelExecution
 #if STATS
 		, const TArray<TStatId>* InStatScopes
 #endif
@@ -313,6 +323,7 @@ struct FVectorVMContext : TThreadSingleton<FVectorVMContext>
 	{
 		Code = InCode;
 		NumInstances = InNumInstances;
+		NumInstancesVectorFloats = (NumInstances + VECTOR_WIDTH_FLOATS - 1) / VECTOR_WIDTH_FLOATS;
 		StartInstance = InStartInstance;
 
 		RandCounters.Reset();
@@ -334,6 +345,39 @@ struct FVectorVMContext : TThreadSingleton<FVectorVMContext>
 		return  ((T*)Meta.OutputRegisters[RegisterIndex]) + Meta.InstanceOffset;
 	}
 
+	int32 GetNumInstances() const { return NumInstances; }
+	int32 GetStartInstance() const { return StartInstance; }
+
+	template<uint32 InstancesPerOp>
+	int32 GetNumLoops() const { return (InstancesPerOp == VECTOR_WIDTH_FLOATS) ? NumInstancesVectorFloats : ((InstancesPerOp == 1) ? NumInstances : Align(NumInstances, InstancesPerOp));	}
+
+	FORCEINLINE uint8 DecodeU8() { return *Code++; }
+#if PLATFORM_SUPPORTS_UNALIGNED_LOADS
+	FORCEINLINE uint16 DecodeU16() { uint16 v = *reinterpret_cast<const uint16*>(Code); Code += sizeof(uint16); return INTEL_ORDER16(v); }
+	FORCEINLINE uint32 DecodeU32() { uint32 v = *reinterpret_cast<const uint32*>(Code); Code += sizeof(uint32); return INTEL_ORDER32(v); }
+	FORCEINLINE uint64 DecodeU64() { uint64 v = *reinterpret_cast<const uint64*>(Code); Code += sizeof(uint64); return INTEL_ORDER64(v); }
+#else
+	FORCEINLINE uint16 DecodeU16() { uint16 v = DecodeU8() <<  8 | DecodeU8(); return INTEL_ORDER16(v); }
+	FORCEINLINE uint16 DecodeU32() { uint32 v = DecodeU8() << 24 | DecodeU8() << 16 | DecodeU8() <<  8 | DecodeU8(); return INTEL_ORDER32(v); }
+	FORCEINLINE uint16 DecodeU64() { uint64 v = DecodeU8() << 54 | DecodeU8() << 48 | DecodeU8() << 40 | DecodeU8() << 32 | DecodeU8() << 24 | DecodeU8() << 16 | DecodeU8() << 8 | DecodeU8(); return INTEL_ORDER64(v); }
+#endif
+	FORCEINLINE uintptr_t DecodePtr() { return (sizeof(uintptr_t) == 4) ? DecodeU32() : DecodeU64(); }
+
+	/** Decode the next operation contained in the bytecode. */
+	FORCEINLINE EVectorVMOp DecodeOp()
+	{
+		return static_cast<EVectorVMOp>(DecodeU8());
+	}
+
+	FORCEINLINE uint8 DecodeSrcOperandTypes()
+	{
+		return DecodeU8();
+	}
+
+	FORCEINLINE bool IsParallelExecution()
+	{
+		return bIsParallelExecution;
+	}
 };
 
 namespace VectorVM
@@ -352,7 +396,8 @@ namespace VectorVM
 	 * Execute VectorVM bytecode.
 	 */
 	VECTORVM_API void Exec(
-		uint8 const* Code,
+		uint8 const* ByteCode,
+		uint8 const* OptimizedByteCode,
 		int32 NumTempRegisters,
 		uint8 const* ConstantTable,
 		TArrayView<FDataSetMeta> DataSetMetaTable,
@@ -364,36 +409,12 @@ namespace VectorVM
 #endif
 	);
 
-	VECTORVM_API void Init();
+	VECTORVM_API void OptimizeByteCode(const uint8* ByteCode, TArray<uint8>& OptimizedCode, TArrayView<uint8> ExternalFunctionRegisterCounts);
 
-	FORCEINLINE uint8 DecodeU8(FVectorVMContext& Context)
-	{
-		return *Context.Code++;
-	}
-
-	FORCEINLINE uint16 DecodeU16(FVectorVMContext& Context)
-	{
-		return ((uint16)DecodeU8(Context) << 8) + DecodeU8(Context);
-	}
-
-	FORCEINLINE uint32 DecodeU32(FVectorVMContext& Context)
-	{
-		return ((uint32)DecodeU8(Context) << 24) + (uint32)(DecodeU8(Context) << 16) + (uint32)(DecodeU8(Context) << 8) + DecodeU8(Context);
-	}
-
-	/** Decode the next operation contained in the bytecode. */
-	FORCEINLINE EVectorVMOp DecodeOp(FVectorVMContext& Context)
-	{
-		return static_cast<EVectorVMOp>(DecodeU8(Context));
-	}
-
-	FORCEINLINE uint8 DecodeSrcOperandTypes(FVectorVMContext& Context)
-	{
-		return DecodeU8(Context);
-	}
+	VECTORVM_API void Init();
 
-#define VVM_EXT_FUNC_INPUT_LOC_BIT (unsigned short)(1<<15)
-#define VVM_EXT_FUNC_INPUT_LOC_MASK (unsigned short)~VVM_EXT_FUNC_INPUT_LOC_BIT
+	#define VVM_EXT_FUNC_INPUT_LOC_BIT (unsigned short)(1<<15)
+	#define VVM_EXT_FUNC_INPUT_LOC_MASK (unsigned short)~VVM_EXT_FUNC_INPUT_LOC_BIT
 
 	template<typename T>
 	struct FUserPtrHandler
@@ -401,7 +422,7 @@ namespace VectorVM
 		int32 UserPtrIdx;
 		T* Ptr;
 		FUserPtrHandler(FVectorVMContext& Context)
-			: UserPtrIdx(*(int32*)(Context.ConstantTable + (DecodeU16(Context))))
+			: UserPtrIdx(*(int32*)(Context.ConstantTable + (Context.DecodeU16())))
 			, Ptr((T*)Context.UserPtrTable[UserPtrIdx])
 		{
 			check(UserPtrIdx != INDEX_NONE);
@@ -435,7 +456,7 @@ namespace VectorVM
 
 		void Init(FVectorVMContext& Context)
 		{
-			InputOffset = DecodeU16(Context);
+			InputOffset = Context.DecodeU16();
 			InputPtr = IsConstant() ? (T*)(Context.ConstantTable + GetOffset()) : (T*)Context.GetTempRegister(GetOffset());
 			AdvanceOffset = IsConstant() ? 0 : 1;
 		}
@@ -471,7 +492,7 @@ namespace VectorVM
 		T* RESTRICT Register;
 	public:
 		FORCEINLINE FExternalFuncRegisterHandler(FVectorVMContext& Context)
-			: RegisterIndex(DecodeU16(Context) & VVM_EXT_FUNC_INPUT_LOC_MASK)
+			: RegisterIndex(Context.DecodeU16() & VVM_EXT_FUNC_INPUT_LOC_MASK)
 			, AdvanceOffset(IsValid() ? 1 : 0)
 		{
 			if (IsValid())
@@ -510,7 +531,7 @@ namespace VectorVM
 		uint16 ConstantIndex;
 		T Constant;
 		FExternalFuncConstHandler(FVectorVMContext& Context)
-			: ConstantIndex(VectorVM::DecodeU16(Context) & VVM_EXT_FUNC_INPUT_LOC_MASK)
+			: ConstantIndex(Context.DecodeU16() & VVM_EXT_FUNC_INPUT_LOC_MASK)
 			, Constant(*((T*)(Context.ConstantTable + ConstantIndex)))
 		{}
 		FORCEINLINE const T& Get() { return Constant; }
