- Fix error when PLATFORM_SUPPORTS_UNALIGNED_LOADS is not valid

#rnx

diff --git a/Engine/Source/Runtime/VectorVM/Public/VectorVM.h b/Engine/Source/Runtime/VectorVM/Public/VectorVM.h
--- a/Engine/Source/Runtime/VectorVM/Public/VectorVM.h
+++ b/Engine/Source/Runtime/VectorVM/Public/VectorVM.h
@@ -358,8 +358,19 @@ public:
 	FORCEINLINE uint64 DecodeU64() { uint64 v = *reinterpret_cast<const uint64*>(Code); Code += sizeof(uint64); return INTEL_ORDER64(v); }
 #else
 	FORCEINLINE uint16 DecodeU16() { uint16 v = DecodeU8() <<  8 | DecodeU8(); return INTEL_ORDER16(v); }
-	FORCEINLINE uint16 DecodeU32() { uint32 v = DecodeU8() << 24 | DecodeU8() << 16 | DecodeU8() <<  8 | DecodeU8(); return INTEL_ORDER32(v); }
-	FORCEINLINE uint16 DecodeU64() { uint64 v = DecodeU8() << 54 | DecodeU8() << 48 | DecodeU8() << 40 | DecodeU8() << 32 | DecodeU8() << 24 | DecodeU8() << 16 | DecodeU8() << 8 | DecodeU8(); return INTEL_ORDER64(v); }
+	FORCEINLINE uint16 DecodeU32() { uint32 v = DecodeU8() << 24ul | DecodeU8() << 16ul | DecodeU8() <<  8ul | DecodeU8(); return INTEL_ORDER32(v); }
+	FORCEINLINE uint16 DecodeU64()
+	{
+		uint64 v = DecodeU8();
+		v = v << 8 | DecodeU8();
+		v = v << 8 | DecodeU8();
+		v = v << 8 | DecodeU8();
+		v = v << 8 | DecodeU8();
+		v = v << 8 | DecodeU8();
+		v = v << 8 | DecodeU8();
+		v = v << 8 | DecodeU8();
+		return INTEL_ORDER64(v);
+	}
 #endif
 	FORCEINLINE uintptr_t DecodePtr() { return (sizeof(uintptr_t) == 4) ? DecodeU32() : DecodeU64(); }
 
