- Fix some static analysis warnings

#rnx

diff --git a/Engine/Source/Runtime/VectorVM/Private/VectorVM.cpp b/Engine/Source/Runtime/VectorVM/Private/VectorVM.cpp
--- a/Engine/Source/Runtime/VectorVM/Private/VectorVM.cpp
+++ b/Engine/Source/Runtime/VectorVM/Private/VectorVM.cpp
@@ -1526,11 +1526,14 @@ struct FScalarKernelAcquireCounterIndex
 			{
 				TBinaryKernelHandler<InternalKernel<true>, FRegisterHandler<int32>, FDataSetCounterHandler, FRegisterHandler<int32>, 1>::Exec(Context);
 			}
-			else
+			else if (SrcOpType == SRCOP_RRC)
 			{
-				check(SrcOpType == SRCOP_RRC);
 				TBinaryKernelHandler<InternalKernel<true>, FRegisterHandler<int32>, FDataSetCounterHandler, FConstantHandler<int32>, 1>::Exec(Context);
 			}
+			else
+			{
+				check(false);
+			}
 		}
 		else
 		{
@@ -1538,11 +1541,14 @@ struct FScalarKernelAcquireCounterIndex
 			{
 				TBinaryKernelHandler<InternalKernel<false>, FRegisterHandler<int32>, FDataSetCounterHandler, FRegisterHandler<int32>, 1>::Exec(Context);
 			}
-			else
+			else if (SrcOpType == SRCOP_RRC)
 			{
-				check(SrcOpType == SRCOP_RRC);
 				TBinaryKernelHandler<InternalKernel<false>, FRegisterHandler<int32>, FDataSetCounterHandler, FConstantHandler<int32>, 1>::Exec(Context);
 			}
+			else
+			{
+				check(false);
+			}
 		}
 	}
 
@@ -2288,11 +2294,9 @@ FString VectorVM::GetOperandLocationName(EVectorVMOperandLocation Location)
 void VectorVM::OptimizeByteCode(const uint8* ByteCode, TArray<uint8>& OptimizedCode, TArrayView<uint8> ExternalFunctionRegisterCounts)
 {
 	OptimizedCode.Empty();
-	if (!PLATFORM_SUPPORTS_UNALIGNED_LOADS || !PLATFORM_LITTLE_ENDIAN)
-	{
-		//-TODO: Support unaligned writes & little endian
-		return;
-	}
+
+//-TODO: Support unaligned writes & little endian
+#if PLATFORM_SUPPORTS_UNALIGNED_LOADS && PLATFORM_LITTLE_ENDIAN
 
 	if ( !GbOptimizeVMByteCode || (ByteCode == nullptr) )
 	{
@@ -2416,6 +2420,7 @@ void VectorVM::OptimizeByteCode(const uint8* ByteCode, TArray<uint8>& OptimizedC
 		}
 	} while (Op != EVectorVMOp::done);
 	Context.Write<FVectorVMExecFunction>(nullptr);
+#endif //PLATFORM_SUPPORTS_UNALIGNED_LOADS && PLATFORM_LITTLE_ENDIAN
 }
 
 #undef VM_FORCEINLINE
