- Fix static analysis warning

#rnx

diff --git a/Engine/Source/Runtime/VectorVM/Private/VectorVM.cpp b/Engine/Source/Runtime/VectorVM/Private/VectorVM.cpp
--- a/Engine/Source/Runtime/VectorVM/Private/VectorVM.cpp
+++ b/Engine/Source/Runtime/VectorVM/Private/VectorVM.cpp
@@ -1517,7 +1517,7 @@ struct FScalarKernelAcquireCounterIndex
 		}
 	};
 
-	template<bool SrcOpType>
+	template<uint32 SrcOpType>
 	static void ExecOptimized(FVectorVMContext& Context)
 	{
 		if (Context.IsParallelExecution())
