- Fix for static analysis warning
- Warning is because a compare flags against a uint32 value and it assumes you should be doing an & == v which isn't correct in this instance

#rnx

diff --git a/Engine/Source/Runtime/VectorVM/Private/VectorVM.cpp b/Engine/Source/Runtime/VectorVM/Private/VectorVM.cpp
--- a/Engine/Source/Runtime/VectorVM/Private/VectorVM.cpp
+++ b/Engine/Source/Runtime/VectorVM/Private/VectorVM.cpp
@@ -1522,32 +1522,20 @@ struct FScalarKernelAcquireCounterIndex
 	{
 		if (Context.IsParallelExecution())
 		{
-			if (SrcOpType == SRCOP_RRR)
-			{
-				TBinaryKernelHandler<InternalKernel<true>, FRegisterHandler<int32>, FDataSetCounterHandler, FRegisterHandler<int32>, 1>::Exec(Context);
-			}
-			else if (SrcOpType == SRCOP_RRC)
-			{
-				TBinaryKernelHandler<InternalKernel<true>, FRegisterHandler<int32>, FDataSetCounterHandler, FConstantHandler<int32>, 1>::Exec(Context);
-			}
-			else
+			switch (SrcOpType)
 			{
-				check(false);
+				case SRCOP_RRR: TBinaryKernelHandler<InternalKernel<true>, FRegisterHandler<int32>, FDataSetCounterHandler, FRegisterHandler<int32>, 1>::Exec(Context); break;
+				case SRCOP_RRC:	TBinaryKernelHandler<InternalKernel<true>, FRegisterHandler<int32>, FDataSetCounterHandler, FConstantHandler<int32>, 1>::Exec(Context); break;
+				default: check(0); break;
 			}
 		}
 		else
 		{
-			if (SrcOpType == SRCOP_RRR)
-			{
-				TBinaryKernelHandler<InternalKernel<false>, FRegisterHandler<int32>, FDataSetCounterHandler, FRegisterHandler<int32>, 1>::Exec(Context);
-			}
-			else if (SrcOpType == SRCOP_RRC)
-			{
-				TBinaryKernelHandler<InternalKernel<false>, FRegisterHandler<int32>, FDataSetCounterHandler, FConstantHandler<int32>, 1>::Exec(Context);
-			}
-			else
+			switch (SrcOpType)
 			{
-				check(false);
+				case SRCOP_RRR: TBinaryKernelHandler<InternalKernel<false>, FRegisterHandler<int32>, FDataSetCounterHandler, FRegisterHandler<int32>, 1>::Exec(Context); break;
+				case SRCOP_RRC:	TBinaryKernelHandler<InternalKernel<false>, FRegisterHandler<int32>, FDataSetCounterHandler, FConstantHandler<int32>, 1>::Exec(Context); break;
+				default: check(0); break;
 			}
 		}
 	}
@@ -1557,14 +1545,11 @@ struct FScalarKernelAcquireCounterIndex
 		const uint32 SrcOpType = Context.BaseContext.DecodeSrcOperandTypes();
 		switch (SrcOpType)
 		{
-			case SRCOP_RRR:
-				Context.Write<FVectorVMExecFunction>(FScalarKernelAcquireCounterIndex::ExecOptimized<SRCOP_RRR>);
-				break;
-			case SRCOP_RRC:
-				Context.Write<FVectorVMExecFunction>(FScalarKernelAcquireCounterIndex::ExecOptimized<SRCOP_RRC>);
-				break;
+			case SRCOP_RRR: Context.Write<FVectorVMExecFunction>(FScalarKernelAcquireCounterIndex::ExecOptimized<SRCOP_RRR>); break;
+			case SRCOP_RRC: Context.Write<FVectorVMExecFunction>(FScalarKernelAcquireCounterIndex::ExecOptimized<SRCOP_RRC>); break;
 			default: check(0); break;
 		}
+
 		// Three registers, note we don't call Optimize on the Kernel since that will write the Exec and we are selecting based upon thread safe or not
 		Context.Write(Context.DecodeU16());
 		Context.Write(Context.DecodeU16());
