- Fix decoding values on platforms that don't support unaligned loads

#rnx

diff --git a/Engine/Source/Runtime/VectorVM/Public/VectorVM.h b/Engine/Source/Runtime/VectorVM/Public/VectorVM.h
--- a/Engine/Source/Runtime/VectorVM/Public/VectorVM.h
+++ b/Engine/Source/Runtime/VectorVM/Public/VectorVM.h
@@ -357,20 +357,9 @@ public:
 	FORCEINLINE uint32 DecodeU32() { uint32 v = *reinterpret_cast<const uint32*>(Code); Code += sizeof(uint32); return INTEL_ORDER32(v); }
 	FORCEINLINE uint64 DecodeU64() { uint64 v = *reinterpret_cast<const uint64*>(Code); Code += sizeof(uint64); return INTEL_ORDER64(v); }
 #else
-	FORCEINLINE uint16 DecodeU16() { uint16 v = DecodeU8() <<  8 | DecodeU8(); return INTEL_ORDER16(v); }
-	FORCEINLINE uint16 DecodeU32() { uint32 v = DecodeU8() << 24ul | DecodeU8() << 16ul | DecodeU8() <<  8ul | DecodeU8(); return INTEL_ORDER32(v); }
-	FORCEINLINE uint16 DecodeU64()
-	{
-		uint64 v = DecodeU8();
-		v = v << 8 | DecodeU8();
-		v = v << 8 | DecodeU8();
-		v = v << 8 | DecodeU8();
-		v = v << 8 | DecodeU8();
-		v = v << 8 | DecodeU8();
-		v = v << 8 | DecodeU8();
-		v = v << 8 | DecodeU8();
-		return INTEL_ORDER64(v);
-	}
+	FORCEINLINE uint16 DecodeU16() { uint16 v = Code[1]; v = v << 8 | Code[0]; Code += 2; return INTEL_ORDER16(v); }
+	FORCEINLINE uint32 DecodeU32() { uint32 v = Code[3]; v = v << 8 | Code[2]; v = v << 8 | Code[1]; v = v << 8 | Code[0]; Code += 4; return INTEL_ORDER32(v); }
+	FORCEINLINE uint64 DecodeU64() { uint64 v = Code[7]; v = v << 8 | Code[6]; v = v << 8 | Code[5]; v = v << 8 | Code[4]; v = v << 8 | Code[3]; v = v << 8 | Code[2]; v = v << 8 | Code[1]; v = v << 8 | Code[0]; Code += 8; return INTEL_ORDER64(v); }
 #endif
 	FORCEINLINE uintptr_t DecodePtr() { return (sizeof(uintptr_t) == 4) ? DecodeU32() : DecodeU64(); }
 
