Optimization to the VectorVM pattern of acquiring indices and outputting stream of attribute data.
Left packs the data with shuffles instead of branching on if a particle is valid, and claims blocks of indices rather than one at a time.

diff --git a/Engine/Source/Runtime/Core/Public/Math/UnrealMathDirectX.h b/Engine/Source/Runtime/Core/Public/Math/UnrealMathDirectX.h
--- a/Engine/Source/Runtime/Core/Public/Math/UnrealMathDirectX.h
+++ b/Engine/Source/Runtime/Core/Public/Math/UnrealMathDirectX.h
@@ -13,6 +13,7 @@
 #include <DirectXMath.h>
 #pragma warning( pop )
 #include <DirectXPackedVector.h>
+#include "Math/sse_mathfun.h"
 
 /*=============================================================================
  *	Helpers:
@@ -970,10 +971,9 @@ FORCEINLINE bool VectorContainsNaNOrInfinite(const VectorRegister& Vec)
 	return !IsFinite;
 }
 
-//TODO: Vectorize
 FORCEINLINE VectorRegister VectorExp(const VectorRegister& X)
 {
-	return MakeVectorRegister(FMath::Exp(VectorGetComponent(X, 0)), FMath::Exp(VectorGetComponent(X, 1)), FMath::Exp(VectorGetComponent(X, 2)), FMath::Exp(VectorGetComponent(X, 3)));
+	return SseMath_exp_ps(X);
 }
 
 //TODO: Vectorize
@@ -982,10 +982,9 @@ FORCEINLINE VectorRegister VectorExp2(const VectorRegister& X)
 	return DirectX::XMVectorExp2(X);
 }
 
-//TODO: Vectorize
 FORCEINLINE VectorRegister VectorLog(const VectorRegister& X)
 {
-	return MakeVectorRegister(FMath::Loge(VectorGetComponent(X, 0)), FMath::Loge(VectorGetComponent(X, 1)), FMath::Loge(VectorGetComponent(X, 2)), FMath::Loge(VectorGetComponent(X, 3)));
+	return SseMath_log_ps(X);
 }
 
 //TODO: Vectorize
@@ -1195,6 +1194,14 @@ FORCEINLINE VectorRegisterInt VectorIntAbs(const VectorRegisterInt& A)
 */
 #define VectorIntLoad1( Ptr )	_mm_shuffle_epi32(_mm_loadu_si128((VectorRegisterInt*)(Ptr)),_MM_SHUFFLE(0,0,0,0))
 
+/**
+ * Shuffles a VectorInt using a provided shuffle mask
+ *
+ * @param Vec		Source vector
+ * @param Mask		Shuffle vector
+ */
+#define VectorIntShuffle( Vec, Mask )	_mm_shuffle_epi8( (Vec), (Mask) )
+
 #endif
 
 
diff --git a/Engine/Source/Runtime/Core/Public/Math/UnrealMathFPU.h b/Engine/Source/Runtime/Core/Public/Math/UnrealMathFPU.h
--- a/Engine/Source/Runtime/Core/Public/Math/UnrealMathFPU.h
+++ b/Engine/Source/Runtime/Core/Public/Math/UnrealMathFPU.h
@@ -1662,6 +1662,27 @@ FORCEINLINE VectorRegisterInt VectorIntLoad1(const void* Ptr)
 		IntSplat);
 }
 
+/**
+ * Shuffles a VectorInt using a provided shuffle mask
+ *
+ * @param Vec		Source vector
+ * @param Mask		Shuffle vector
+ */
+FORCEINLINE VectorRegisterInt VectorIntShuffle(const VectorRegisterInt& Vec, const VectorRegisterInt& Mask)
+{
+	VectorRegisterInt Result;
+	const int8* VecBytes = reinterpret_cast<const int8*>(&Vec);
+	const int8* MaskBytes = reinterpret_cast<const int8*>(&Mask);
+	int8* ResultBytes = reinterpret_cast<int8*>(&Result);
+
+	for (int32 i = 0; i < sizeof(VectorRegisterInt); ++i)
+	{
+		ResultBytes[i] = (MaskBytes[i] < 0) ? 0 : VecBytes[MaskBytes[i] % 16];
+	}
+
+	return Result;
+}
+
 /**
  * These functions return a vector mask to indicate which components pass the comparison.
  * Each component is 0xffffffff if it passes, 0x00000000 if it fails.
diff --git a/Engine/Source/Runtime/Core/Public/Math/UnrealMathNeon.h b/Engine/Source/Runtime/Core/Public/Math/UnrealMathNeon.h
--- a/Engine/Source/Runtime/Core/Public/Math/UnrealMathNeon.h
+++ b/Engine/Source/Runtime/Core/Public/Math/UnrealMathNeon.h
@@ -1653,6 +1653,18 @@ FORCEINLINE VectorRegisterInt VectorIntSelect(const VectorRegisterInt& Mask, con
 */
 #define VectorIntLoad1( Ptr )	vld1q_dup_s32((int32*)(Ptr))
 
+/**
+ * Shuffles a VectorInt using a provided shuffle mask
+ *
+ * @param Vec		Source vector
+ * @param Mask		Shuffle vector
+ */
+FORCEINLINE VectorRegisterInt VectorIntShuffle(const VectorRegisterInt& Vec, const VectorRegisterInt& Mask)
+{
+	uint8x8x2_t VecSplit = (uint8x8x2_t) { vget_low_u8(Vec), vget_high_u8(Vec) };
+	return vcombine_u8(vtbl2_u8(VecSplit, vget_low_u8(Mask)), vtbl2_u8(VecSplit, vget_high_u8(Mask)));
+}
+
 // To be continued...
 
 PRAGMA_ENABLE_SHADOW_VARIABLE_WARNINGS
diff --git a/Engine/Source/Runtime/Core/Public/Math/UnrealMathSSE.h b/Engine/Source/Runtime/Core/Public/Math/UnrealMathSSE.h
--- a/Engine/Source/Runtime/Core/Public/Math/UnrealMathSSE.h
+++ b/Engine/Source/Runtime/Core/Public/Math/UnrealMathSSE.h
@@ -1228,7 +1228,6 @@ FORCEINLINE VectorRegister VectorStep(const VectorRegister& X)
 	return VectorSelect(Mask, (GlobalVectorConstants::FloatOne), (GlobalVectorConstants::FloatZero));
 }
 
-//TODO: Vectorize
 FORCEINLINE VectorRegister VectorExp(const VectorRegister& X)
 {
 	return SseMath_exp_ps(X);
@@ -1240,7 +1239,6 @@ FORCEINLINE VectorRegister VectorExp2(const VectorRegister& X)
 	return MakeVectorRegister(FMath::Exp2(VectorGetComponent(X, 0)), FMath::Exp2(VectorGetComponent(X, 1)), FMath::Exp2(VectorGetComponent(X, 2)), FMath::Exp2(VectorGetComponent(X, 3)));
 }
 
-//TODO: Vectorize
 FORCEINLINE VectorRegister VectorLog(const VectorRegister& X)
 {
 	return SseMath_log_ps(X);
@@ -1490,5 +1488,15 @@ FORCEINLINE VectorRegisterInt VectorIntAbs(const VectorRegisterInt& A)
 * @return		VectorRegisterInt(*Ptr, *Ptr, *Ptr, *Ptr)
 */
 #define VectorIntLoad1( Ptr )	_mm_shuffle_epi32(_mm_loadu_si128((VectorRegisterInt*)(Ptr)),_MM_SHUFFLE(0,0,0,0))
+
+/**
+ * Shuffles a VectorInt using a provided shuffle mask
+ *
+ * @param Vec		Source vector
+ * @param Mask		Shuffle vector
+ */
+#define VectorIntShuffle( Vec, Mask )	_mm_shuffle_epi8( (Vec), (Mask) )
+
+
 #endif
 
diff --git a/Engine/Source/Runtime/VectorVM/Private/VectorVM.cpp b/Engine/Source/Runtime/VectorVM/Private/VectorVM.cpp
--- a/Engine/Source/Runtime/VectorVM/Private/VectorVM.cpp
+++ b/Engine/Source/Runtime/VectorVM/Private/VectorVM.cpp
@@ -1,4 +1,4 @@
-// Copyright 1998-2019 Epic Games, Inc. All Rights Reserved.
+#// Copyright 1998-2019 Epic Games, Inc. All Rights Reserved.
 
 #include "VectorVM.h"
 #include "Modules/ModuleManager.h"
@@ -41,6 +41,49 @@ DEFINE_LOG_CATEGORY_STATIC(LogVectorVM, All, All);
 #define SRCOP_CCR (OP2_CONST | OP1_CONST | OP_REGISTER)
 #define SRCOP_CCC (OP2_CONST | OP1_CONST | OP0_CONST)
 
+namespace VectorVMConstants
+{
+	static const VectorRegisterInt VectorStride = MakeVectorRegisterInt(VECTOR_WIDTH_FLOATS, VECTOR_WIDTH_FLOATS, VECTOR_WIDTH_FLOATS, VECTOR_WIDTH_FLOATS);
+
+	// for generating shuffle masks given input {A, B, C, D}
+	constexpr uint32 ShufMaskIgnore = 0xFFFFFFFF;
+	constexpr uint32 ShufMaskA = 0x03020100;
+	constexpr uint32 ShufMaskB = 0x07060504;
+	constexpr uint32 ShufMaskC = 0x0B0A0908;
+	constexpr uint32 ShufMaskD = 0x0F0E0D0C;
+
+	static const VectorRegisterInt RegisterShuffleMask[] =
+	{
+		MakeVectorRegisterInt(ShufMaskIgnore, ShufMaskIgnore, ShufMaskIgnore, ShufMaskIgnore), // 0000
+		MakeVectorRegisterInt(ShufMaskD, ShufMaskIgnore, ShufMaskIgnore, ShufMaskIgnore), // 0001
+		MakeVectorRegisterInt(ShufMaskC, ShufMaskIgnore, ShufMaskIgnore, ShufMaskIgnore), // 0010
+		MakeVectorRegisterInt(ShufMaskC, ShufMaskD, ShufMaskIgnore, ShufMaskIgnore), // 0011
+		MakeVectorRegisterInt(ShufMaskB, ShufMaskIgnore, ShufMaskIgnore, ShufMaskIgnore), // 0100
+		MakeVectorRegisterInt(ShufMaskB, ShufMaskD, ShufMaskIgnore, ShufMaskIgnore), // 0101
+		MakeVectorRegisterInt(ShufMaskB, ShufMaskC, ShufMaskIgnore, ShufMaskIgnore), // 0110
+		MakeVectorRegisterInt(ShufMaskB, ShufMaskC, ShufMaskD, ShufMaskIgnore), // 0111
+		MakeVectorRegisterInt(ShufMaskA, ShufMaskIgnore, ShufMaskIgnore, ShufMaskIgnore), // 1000
+		MakeVectorRegisterInt(ShufMaskA, ShufMaskD, ShufMaskIgnore, ShufMaskIgnore), // 1001
+		MakeVectorRegisterInt(ShufMaskA, ShufMaskC, ShufMaskIgnore, ShufMaskIgnore), // 1010
+		MakeVectorRegisterInt(ShufMaskA, ShufMaskC, ShufMaskD, ShufMaskIgnore), // 1011
+		MakeVectorRegisterInt(ShufMaskA, ShufMaskB, ShufMaskIgnore, ShufMaskIgnore), // 1100
+		MakeVectorRegisterInt(ShufMaskA, ShufMaskB, ShufMaskD, ShufMaskIgnore), // 1101
+		MakeVectorRegisterInt(ShufMaskA, ShufMaskB, ShufMaskC, ShufMaskIgnore), // 1110
+		MakeVectorRegisterInt(ShufMaskA, ShufMaskB, ShufMaskC, ShufMaskD), // 1111
+	};
+
+	constexpr uint32 cOne = 0xFFFFFFFFU;
+	constexpr uint32 cZero = 0x00000000U;
+	static const VectorRegister RemainderMask[] =
+	{
+		MakeVectorRegister(cZero, cZero, cZero, cZero), // 0 remaining
+		MakeVectorRegister(cOne, cZero, cZero, cZero), // 1 remaining
+		MakeVectorRegister(cOne, cOne, cZero, cZero), // 2 remaining
+		MakeVectorRegister(cOne, cOne, cOne, cZero), // 3 remaining
+		MakeVectorRegister(cOne, cOne, cOne, cOne), // 4 remaining
+	};
+};
+
 //Temporarily locking the free table until we can implement a lock free algorithm. UE-65856
 FORCEINLINE void FDataSetMeta::LockFreeTable()
 {
@@ -98,7 +141,7 @@ static FAutoConsoleVariableRef CVarParallelVVMInstancesPerChunk(
 	ECVF_ReadOnly
 );
 
-static int32 GbOptimizeVMByteCode = 0;
+static int32 GbOptimizeVMByteCode = 1;
 static FAutoConsoleVariableRef CVarbOptimizeVMByteCode(
 	TEXT("vm.OptimizeVMByteCode"),
 	GbOptimizeVMByteCode,
@@ -106,7 +149,7 @@ static FAutoConsoleVariableRef CVarbOptimizeVMByteCode(
 	ECVF_Default
 );
 
-static int32 GbUseOptimizedVMByteCode = 0;
+static int32 GbUseOptimizedVMByteCode = 1;
 static FAutoConsoleVariableRef CVarbUseOptimizedVMByteCode(
 	TEXT("vm.UseOptimizedVMByteCode"),
 	GbUseOptimizedVMByteCode,
@@ -122,6 +165,14 @@ static FAutoConsoleVariableRef CVarbSafeOptimizedKernels(
 	ECVF_Default
 );
 
+static int32 GbBatchPackVMOutput = 1;
+static FAutoConsoleVariableRef CVarbBatchPackVMOutput(
+	TEXT("vm.BatchPackedVMOutput"),
+	GbBatchPackVMOutput,
+	TEXT("If > 0 output elements will be packed and batched branch free.\n"),
+	ECVF_Default
+);
+
 //////////////////////////////////////////////////////////////////////////
 //  VM Code Optimizer Context
 
@@ -129,16 +180,14 @@ typedef void(*FVectorVMExecFunction)(FVectorVMContext&);
 
 struct FVectorVMCodeOptimizerContext
 {
+	typedef EVectorVMOp(*OptimizeVMFunction)(EVectorVMOp, FVectorVMCodeOptimizerContext&);
+
 	explicit FVectorVMCodeOptimizerContext(FVectorVMContext& InBaseContext, const uint8* ByteCode, TArray<uint8>& InOptimizedCode, TArrayView<uint8> InExternalFunctionRegisterCounts)
 		: BaseContext(InBaseContext)
 		, OptimizedCode(InOptimizedCode)
 		, ExternalFunctionRegisterCounts(InExternalFunctionRegisterCounts)
 	{
-#if STATS
-		BaseContext.PrepareForExec(0, nullptr, nullptr, nullptr, TArrayView<FDataSetMeta>(), 0, false, &EmptyStatScopes);
-#else
 		BaseContext.PrepareForExec(0, nullptr, nullptr, nullptr, TArrayView<FDataSetMeta>(), 0, false);
-#endif
 		BaseContext.PrepareForChunk(ByteCode, 0, 0);
 	}
 	FVectorVMCodeOptimizerContext(const FVectorVMCodeOptimizerContext&) = delete;
@@ -159,13 +208,31 @@ struct FVectorVMCodeOptimizerContext
 		reinterpret_cast<T&>(OptimizedCode[OptimizedCode.AddUninitialized(sizeof(T))]) = v;
 	}
 
+	struct FOptimizerCodeState
+	{
+		uint8 const* BaseContextCode;
+		int32 OptimizedCodeLength;
+	};
+
+	FOptimizerCodeState CreateCodeState()
+	{
+		FOptimizerCodeState State;
+		State.BaseContextCode = BaseContext.Code;
+		State.OptimizedCodeLength = OptimizedCode.Num();
+
+		return State;
+	}
+
+	void RollbackCodeState(const FOptimizerCodeState& State)
+	{
+		BaseContext.Code = State.BaseContextCode;
+		OptimizedCode.SetNum(State.OptimizedCodeLength, false /* allowShrink */);
+	}
+
 	FVectorVMContext&		BaseContext;
 	TArray<uint8>&			OptimizedCode;
 	const TArrayView<uint8>	ExternalFunctionRegisterCounts;
 	const int32				StartInstance = 0;
-#if STATS
-	TArray<TStatId>			EmptyStatScopes;
-#endif
 };
 
 //////////////////////////////////////////////////////////////////////////
@@ -293,20 +360,12 @@ void FVectorVMContext::PrepareForExec(
 	TArrayView<FDataSetMeta> InDataSetMetaTable,
 	int32 MaxNumInstances,
 	bool bInParallelExecution
-#if STATS
-	, const TArray<TStatId>* InStatScopes
-#endif
 )
 {
 	NumTempRegisters = InNumTempRegisters;
 	ConstantTable = InConstantTable;
 	ExternalFunctionTable = InExternalFunctionTable;
 	UserPtrTable = InUserPtrTable;
-#if STATS
-	check(InStatScopes);
-	StatScopes = InStatScopes;
-	StatCounterStack.Reserve(StatScopes->Num());
-#endif
 
 	TempRegisterSize = Align(MaxNumInstances * VectorVM::MaxInstanceSizeBytes, PLATFORM_CACHE_LINE_SIZE);
 	TempBufferSize = TempRegisterSize * NumTempRegisters;
@@ -323,6 +382,15 @@ void FVectorVMContext::PrepareForExec(
 	bIsParallelExecution = bInParallelExecution;
 }
 
+#if STATS
+void FVectorVMContext::SetStatScopes(const TArray<TStatId>* InStatScopes)
+{
+	check(InStatScopes);
+	StatScopes = InStatScopes;
+	StatCounterStack.Reserve(StatScopes->Num());
+}
+#endif
+
 void FVectorVMContext::FinishExec()
 {
 	//At the end of executing each chunk we can push any thread local temporary data out to the main storage with locks or atomics.
@@ -960,7 +1028,8 @@ struct FVectorKernelEnterStatScope
 		Context.Write<FVectorVMExecFunction>(Exec);
 		FConstantHandler<int32>::Optimize(Context);
 #else
-		FConstantHandler<int32> ScopeIdx(Context.BaseContext);
+		// just skip the op if we don't have stats enabled
+		FConstantHandler<int32>(Context.BaseContext);
 #endif
 	}
 
@@ -968,7 +1037,7 @@ struct FVectorKernelEnterStatScope
 	{
 		FConstantHandler<int32> ScopeIdx(Context);
 #if STATS
-		if (GbDetailedVMScriptStats)
+		if (GbDetailedVMScriptStats && Context.StatScopes)
 		{
 			int32 CounterIdx = Context.StatCounterStack.AddDefaulted(1);
 			Context.StatCounterStack[CounterIdx].Start((*Context.StatScopes)[ScopeIdx.Get()]);
@@ -1023,13 +1092,13 @@ struct FVectorKernelRandomGauss : public TBinaryVectorKernel<FVectorKernelRandom
 			Context.RandStream.GetFraction(),
 			Context.RandStream.GetFraction());
 
-		Result = VectorSubtract(Result, MakeVectorRegister(0.5f, 0.5f, 0.5f, 0.5f));
+		Result = VectorSubtract(Result, GlobalVectorConstants::FloatOneHalf);
 		Result = VectorMultiply(MakeVectorRegister(3.0f, 3.0f, 3.0f, 3.0f), Result);
 
 		// taylor series gaussian approximation
 		const VectorRegister SPi2 = VectorReciprocal(VectorReciprocalSqrt(MakeVectorRegister(2 * PI, 2 * PI, 2 * PI, 2 * PI)));
 		VectorRegister Gauss = VectorReciprocal(SPi2);
-		VectorRegister Div = VectorMultiply(MakeVectorRegister(2.0f, 2.0f, 2.0f, 2.0f), SPi2);
+		VectorRegister Div = VectorMultiply(GlobalVectorConstants::FloatTwo, SPi2);
 		Gauss = VectorSubtract(Gauss, VectorDivide(VectorMultiply(Result, Result), Div));
 		Div = VectorMultiply(MakeVectorRegister(8.0f, 8.0f, 8.0f, 8.0f), SPi2);
 		Gauss = VectorAdd(Gauss, VectorDivide(VectorPow(MakeVectorRegister(4.0f, 4.0f, 4.0f, 4.0f), Result), Div));
@@ -1226,10 +1295,9 @@ struct FVectorKernelNoise : public TUnaryVectorKernel<FVectorKernelNoise>
 
 	static void VM_FORCEINLINE DoKernel(FVectorVMContext& Context, VectorRegister* RESTRICT Dst, VectorRegister Src0)
 	{
-		const VectorRegister One = MakeVectorRegister(1.0f, 1.0f, 1.0f, 1.0f);
 		const VectorRegister VecSize = MakeVectorRegister(16.0f, 16.0f, 16.0f, 16.0f);
 
-		*Dst = MakeVectorRegister(0.0f, 0.0f, 0.0f, 0.0f);
+		*Dst = GlobalVectorConstants::FloatZero;
 		
 		for (uint32 i = 1; i < 2; i++)
 		{
@@ -1243,7 +1311,7 @@ struct FVectorKernelNoise : public TUnaryVectorKernel<FVectorKernelNoise>
 
 			VectorRegister Frac = VectorFractional(Coords);
 			VectorRegister Alpha = VectorReplicate(Frac, 0);
-			VectorRegister OneMinusAlpha = VectorSubtract(One, Alpha);
+			VectorRegister OneMinusAlpha = VectorSubtract(GlobalVectorConstants::FloatOne, Alpha);
 			
 			VectorRegister XV1 = VectorMultiplyAdd(RandomTable[Cx][Cy][Cz], Alpha, VectorMultiply(RandomTable[Cx+1][Cy][Cz], OneMinusAlpha));
 			VectorRegister XV2 = VectorMultiplyAdd(RandomTable[Cx][Cy+1][Cz], Alpha, VectorMultiply(RandomTable[Cx+1][Cy+1][Cz], OneMinusAlpha));
@@ -1251,12 +1319,12 @@ struct FVectorKernelNoise : public TUnaryVectorKernel<FVectorKernelNoise>
 			VectorRegister XV4 = VectorMultiplyAdd(RandomTable[Cx][Cy+1][Cz+1], Alpha, VectorMultiply(RandomTable[Cx+1][Cy+1][Cz+1], OneMinusAlpha));
 
 			Alpha = VectorReplicate(Frac, 1);
-			OneMinusAlpha = VectorSubtract(One, Alpha);
+			OneMinusAlpha = VectorSubtract(GlobalVectorConstants::FloatOne, Alpha);
 			VectorRegister YV1 = VectorMultiplyAdd(XV1, Alpha, VectorMultiply(XV2, OneMinusAlpha));
 			VectorRegister YV2 = VectorMultiplyAdd(XV3, Alpha, VectorMultiply(XV4, OneMinusAlpha));
 
 			Alpha = VectorReplicate(Frac, 2);
-			OneMinusAlpha = VectorSubtract(One, Alpha);
+			OneMinusAlpha = VectorSubtract(GlobalVectorConstants::FloatOne, Alpha);
 			VectorRegister ZV = VectorMultiplyAdd(YV1, Alpha, VectorMultiply(YV2, OneMinusAlpha));
 
 			*Dst = VectorAdd(*Dst, ZV);
@@ -2025,8 +2093,10 @@ struct FVectorKernelBoolToInt : TUnaryKernel<FVectorKernelBoolToInt, FRegisterHa
 	}
 };
 
+#if WITH_EDITOR
 UEnum* g_VectorVMEnumStateObj = nullptr;
 UEnum* g_VectorVMEnumOperandObj = nullptr;
+#endif
 
 
 void VectorVM::Init()
@@ -2034,8 +2104,10 @@ void VectorVM::Init()
 	static bool Inited = false;
 	if (Inited == false)
 	{
+#if WITH_EDITOR
 		g_VectorVMEnumStateObj = StaticEnum<EVectorVMOp>();
 		g_VectorVMEnumOperandObj = StaticEnum<EVectorVMOperandLocation>();
+#endif
 
 		// random noise
 		float TempTable[17][17][17];
@@ -2145,11 +2217,10 @@ void VectorVM::Exec(
 		//SCOPE_CYCLE_COUNTER(STAT_VVMExecChunk);
 
 		FVectorVMContext& Context = FVectorVMContext::Get();
-		Context.PrepareForExec(NumTempRegisters, ConstantTable, ExternalFunctionTable, UserPtrTable, DataSetMetaTable, MaxInstances, bParallel
+		Context.PrepareForExec(NumTempRegisters, ConstantTable, ExternalFunctionTable, UserPtrTable, DataSetMetaTable, MaxInstances, bParallel);
 #if STATS
-			, &StatScopes
+		Context.SetStatScopes(&StatScopes);
 #endif
-		);
 
 		// Process one chunk at a time.
 		int32 ChunkIdx = BatchIdx * ChunksPerBatch;
@@ -2341,6 +2412,467 @@ FString VectorVM::GetOperandLocationName(EVectorVMOperandLocation Location)
 }
 #endif
 
+// Optimization managed by GbBatchPackVMOutput via PackedOutputOptimization()
+// Looks for the common pattern of an acquireindex op followed by a number of associated outputdata_32bit ops.  The
+// stock operation is to write an index into a temporary register, and then have the different outputs streams
+// write into the indexed location.  This optimization does a number of things:
+// -first we check if 'validity' is uniform or not, if it is we can have a fast path of both figuring out how many
+// indices we need, as well as how to write the output (if we find that they are all invalid, then we don't need to do anything!)
+// -if we need to evaluate the validity of each element we quickly count up the number (with vector intrinsics) and
+// grab a block of the indices (rather than one at a time)
+// -rather than storing the indices to use, we store a int8 mask which indicates a valid flag for each of the next 4 samples
+// -outputs are then written to depending on their source and their frequency:
+//		-uniform sources will be splatted to all valid entries
+//		-variable sources will be packed into the available slots
+struct FBatchedWriteIndexedOutput
+{
+	// functor for copying a source register to an output register
+	struct FCopyOp
+	{
+		void VM_FORCEINLINE operator()(FVectorVMContext& Context, uint16 DataSetIndex)
+		{
+			FRegisterHandler<int32> SourceRegister(Context);
+			const uint16 DestRegisterIdx = Context.DecodeU16();
+
+			int32* DestReg = Context.GetOutputRegister<int32>(DataSetIndex, DestRegisterIdx) + Context.ValidInstanceIndexStart;
+
+			FMemory::StreamingMemcpy(DestReg, SourceRegister.GetDest(), sizeof(int32) * Context.ValidInstanceCount);
+		}
+	};
+
+	// functor for splatting a constant value to an output register
+	template<typename InputHandler>
+	struct FSplatOp
+	{
+		void VM_FORCEINLINE operator()(FVectorVMContext& Context, uint16 DataSetIndex)
+		{
+			InputHandler SourceRegister(Context);
+			const uint16 DestRegisterIdx = Context.DecodeU16();
+
+			int32* DestReg = Context.GetOutputRegister<int32>(DataSetIndex, DestRegisterIdx) + Context.ValidInstanceIndexStart;
+
+			const int32 SourceValue = SourceRegister.Get();
+			const int32 InstanceVectorCount = FMath::DivideAndRoundDown(Context.ValidInstanceCount, VECTOR_WIDTH_FLOATS);
+
+			if (InstanceVectorCount)
+			{
+				const VectorRegisterInt SplatValue = MakeVectorRegisterInt(SourceValue, SourceValue, SourceValue, SourceValue);
+
+				for (int32 VectorIt = 0; VectorIt < InstanceVectorCount; ++VectorIt)
+				{
+					VectorIntStore(SplatValue, DestReg + VectorIt * VECTOR_WIDTH_FLOATS);
+				}
+			}
+
+			for (int32 InstanceIt = InstanceVectorCount * VECTOR_WIDTH_FLOATS; InstanceIt < Context.ValidInstanceCount; ++InstanceIt)
+			{
+				DestReg[InstanceIt] = SourceValue;
+			}
+		}
+	};
+
+	// performs the operation of copying data from a temporary register to an output register under the assumption
+	// that the validity of each instance is uniform (valid or not).
+	template<typename PopulateOp>
+	static VM_FORCEINLINE void DoRegisterKernelFixedValid(FVectorVMContext& Context)
+	{
+		const uint16 DataSetIndex = Context.DecodeU16();
+		Context.DecodeU16(); // DestIndexRegisterIdx
+		const uint16 AccumulatedOpCount = Context.DecodeU16();
+
+		// if none of the instances are valid, then don't bother writing anything
+		if (!Context.ValidInstanceCount)
+		{
+			// todo we should early out of this case rather than keep parsing the code
+			for (uint16 OpIt = 0; OpIt < AccumulatedOpCount; ++OpIt)
+			{
+				FRegisterHandler<int32> Dummy(Context);
+				Context.DecodeU16(); // DestRegisterIdx
+			}
+
+			return;
+		}
+
+		// for each of our ops, copy the data from the working register to the output
+		const int32 DataSize = sizeof(int32) * Context.ValidInstanceCount;
+
+		for (uint16 OpIt = 0; OpIt < AccumulatedOpCount; ++OpIt)
+		{
+			PopulateOp()(Context, DataSetIndex);
+		}
+	}
+
+	// performs the operation of copying data from a temporary register to an output register without foreknowledge
+	// of the validity of individual instances
+	static VM_FORCEINLINE void DoRegisterKernelVariableValid(FVectorVMContext& Context)
+	{
+		// if we found that all of the instances are valid, then just run the fixed version
+		if (Context.ValidInstanceUniform)
+		{
+			DoRegisterKernelFixedValid<FCopyOp>(Context);
+			return;
+		}
+
+		const uint16 DataSetIndex = Context.DecodeU16();
+		const uint16 DestIndexRegisterIdx = Context.DecodeU16();
+		const uint16 AccumulatedOpCount = Context.DecodeU16();
+
+		FDataSetMeta& DataSetMeta = Context.GetDataSetMeta(DataSetIndex);
+		
+		const int8* DestIndexReg = reinterpret_cast<const int8*>(Context.GetTempRegister(DestIndexRegisterIdx));
+		const int32 LoopCount = FMath::DivideAndRoundUp(Context.NumInstances, VECTOR_WIDTH_FLOATS);
+
+		uint16 OpIt = 0;
+
+		//
+		// VectorIntStore(		- unaligned writes of 16 bytes to our Destination; note that this maneuver requires us to have
+		//						our output buffers padded out to 16 bytes!
+		//	VectorIntShuffle(	- swizzle our source register to pack the valid entries at the beginning, with 0s at the end
+		//    Source,			- source data
+		//    ShuffleMask),		- result of the VectorMaskBits done in the acquireindex, int8/VectorRegister of input
+		//  Destination);
+
+		constexpr int32 OpsPerLoop = 4;
+		for (OpIt = 0; (OpIt + OpsPerLoop) < AccumulatedOpCount; OpIt += OpsPerLoop)
+		{
+			const RegisterType* Source0 = FRegisterHandler<RegisterType>(Context).GetDest();
+			int32* DestReg0 = Context.GetOutputRegister<int32>(DataSetIndex, Context.DecodeU16()) + Context.ValidInstanceIndexStart;
+
+			const RegisterType* Source1 = FRegisterHandler<RegisterType>(Context).GetDest();
+			int32* DestReg1 = Context.GetOutputRegister<int32>(DataSetIndex, Context.DecodeU16()) + Context.ValidInstanceIndexStart;
+
+			const RegisterType* Source2 = FRegisterHandler<RegisterType>(Context).GetDest();
+			int32* DestReg2 = Context.GetOutputRegister<int32>(DataSetIndex, Context.DecodeU16()) + Context.ValidInstanceIndexStart;
+
+			const RegisterType* Source3 = FRegisterHandler<RegisterType>(Context).GetDest();
+			int32* DestReg3 = Context.GetOutputRegister<int32>(DataSetIndex, Context.DecodeU16()) + Context.ValidInstanceIndexStart;
+
+			for (int32 LoopIt = 0; LoopIt < LoopCount; ++LoopIt)
+			{
+				const int8 ShuffleMask = DestIndexReg[LoopIt];
+				const int8 AdvanceCount = FMath::CountBits(ShuffleMask);
+
+				VectorIntStore(VectorIntShuffle(Source0[LoopIt], VectorVMConstants::RegisterShuffleMask[ShuffleMask]), DestReg0);
+				VectorIntStore(VectorIntShuffle(Source1[LoopIt], VectorVMConstants::RegisterShuffleMask[ShuffleMask]), DestReg1);
+				VectorIntStore(VectorIntShuffle(Source2[LoopIt], VectorVMConstants::RegisterShuffleMask[ShuffleMask]), DestReg2);
+				VectorIntStore(VectorIntShuffle(Source3[LoopIt], VectorVMConstants::RegisterShuffleMask[ShuffleMask]), DestReg3);
+
+				DestReg0 += AdvanceCount;
+				DestReg1 += AdvanceCount;
+				DestReg2 += AdvanceCount;
+				DestReg3 += AdvanceCount;
+			}
+		}
+
+		for (; OpIt < AccumulatedOpCount; ++OpIt)
+		{
+			const RegisterType* Source = FRegisterHandler<RegisterType>(Context).GetDest();
+			int32* DestReg0 = Context.GetOutputRegister<int32>(DataSetIndex, Context.DecodeU16()) + Context.ValidInstanceIndexStart;
+
+			for (int32 LoopIt = 0; LoopIt < LoopCount; ++LoopIt)
+			{
+				const int8 ShuffleMask = DestIndexReg[LoopIt];
+				const int8 AdvanceCount = FMath::CountBits(ShuffleMask);
+
+				VectorIntStore(VectorIntShuffle(Source[LoopIt], VectorVMConstants::RegisterShuffleMask[ShuffleMask]), DestReg0);
+
+				DestReg0 += AdvanceCount;
+			}
+		}
+	}
+
+	// acquires a batch of indices from the provided CounterHandler.  If we're running in parallel, then we'll need to use
+	// atomics to guarantee our place in the list of indices.
+	template<bool bParallel>
+	static VM_FORCEINLINE void AcquireCounterIndex(FVectorVMContext& Context, FDataSetCounterHandler& CounterHandler, int32 AcquireCount)
+	{
+		if (AcquireCount)
+		{
+			int32* CounterHandlerIndex = CounterHandler.Get();
+			int32 StartIndex = INDEX_NONE;
+
+			if (bParallel)
+			{
+				StartIndex = FPlatformAtomics::InterlockedAdd(CounterHandlerIndex, AcquireCount);
+			}
+			else
+			{
+				StartIndex = *CounterHandlerIndex;
+				*CounterHandlerIndex = StartIndex + AcquireCount;
+			}
+
+			// increment StartIndex, since CounterHandlerIndex starts at INDEX_NONE
+			Context.ValidInstanceIndexStart = StartIndex + 1;
+		}
+
+		Context.ValidInstanceCount = AcquireCount;
+		Context.ValidInstanceUniform = !AcquireCount || (Context.NumInstances == AcquireCount);
+	}
+
+	// evaluates a register to evaluate which instances are valid or not; will read 4 entries at a time and generate a
+	// a mask for which entries are valid as well as an overall count
+	template<bool bParallel>
+	static void HandleRegisterValidIndices(FVectorVMContext& Context)
+	{
+		FDataSetCounterHandler CounterHandler(Context);
+		FRegisterHandler<VectorRegister> ValidReader(Context);
+		FRegisterHandler<int8> Dst(Context);
+
+		int8* DestAddr = Dst.GetDest();
+
+		// we can process VECTOR_WIDTH_FLOATS entries at a time, generating a int8 mask for each set of 4 indicating
+		// which are valid
+		const int32 LoopCount = FMath::DivideAndRoundUp(Context.NumInstances, VECTOR_WIDTH_FLOATS);
+
+		int32 Remainder = Context.NumInstances;
+		int32 ValidCount = 0;
+		for (int32 LoopIt = 0; LoopIt < LoopCount; ++LoopIt)
+		{
+			// input register needs to be padded to allow for 16 byte reads; but mask out the ones beyond NumInstances
+			const VectorRegister Mask = VectorVMConstants::RemainderMask[FMath::Min(VECTOR_WIDTH_FLOATS, Remainder)];
+
+			const int8 ValidMask = static_cast<int8>(VectorMaskBits(VectorSelect(Mask, ValidReader.GetAndAdvance(), GlobalVectorConstants::FloatZero)));
+			ValidCount += FMath::CountBits(ValidMask);
+
+			DestAddr[LoopIt] = ValidMask;
+
+			Remainder -= VECTOR_WIDTH_FLOATS;
+		}
+
+		// grab our batch of indices
+		AcquireCounterIndex<bParallel>(Context, CounterHandler, ValidCount);
+	}
+
+	// evaluates the uniform check and grab the appropriate number of indices
+	template<typename ValidReaderType, bool bParallel>
+	static VM_FORCEINLINE void HandleUniformValidIndices(FVectorVMContext& Context)
+	{
+		FDataSetCounterHandler CounterHandler(Context);
+		ValidReaderType ValidReader(Context);
+
+		if (ValidReader.Get())
+		{
+			AcquireCounterIndex<bParallel>(Context, CounterHandler, Context.NumInstances);
+		}
+	}
+
+	template<uint8 SrcOpType>
+	static VM_FORCEINLINE void IndexExecOptimized(FVectorVMContext& Context)
+	{
+		if (Context.IsParallelExecution())
+		{
+			switch (SrcOpType)
+			{
+			case SRCOP_RRR: HandleRegisterValidIndices<true>(Context); break;
+			case SRCOP_RRC:	HandleUniformValidIndices<FConstantHandler<int32>, true>(Context); break;
+			default: check(0); break;
+			}
+		}
+		else
+		{
+			switch (SrcOpType)
+			{
+			case SRCOP_RRR: HandleRegisterValidIndices<false>(Context); break;
+			case SRCOP_RRC:	HandleUniformValidIndices<FConstantHandler<int32>, false>(Context); break;
+			default: check(0); break;
+			}
+		}
+	}
+
+	void OptimizeAcquireIndex(FVectorVMCodeOptimizerContext& Context)
+	{
+		const uint32 SrcOpType = Context.BaseContext.DecodeSrcOperandTypes();
+
+		AcquireIndexConstant = !!(SrcOpType & OP0_CONST);
+
+		switch (SrcOpType)
+		{
+		case SRCOP_RRR: Context.Write<FVectorVMExecFunction>(IndexExecOptimized<SRCOP_RRR>); break;
+		case SRCOP_RRC: Context.Write<FVectorVMExecFunction>(IndexExecOptimized<SRCOP_RRC>); break;
+		default: check(0); break;
+		}
+
+		DataSetCounterIndex = Context.DecodeU16();
+		ValidTestRegisterIndex = Context.DecodeU16();
+		WorkingRegisterIndex = Context.DecodeU16();
+
+		Context.Write(DataSetCounterIndex);
+		Context.Write(ValidTestRegisterIndex);
+
+		// we only need the working register if we've got non-uniform data
+		if (SrcOpType == SRCOP_RRR)
+		{
+			Context.Write(WorkingRegisterIndex);
+		}
+	}
+
+	bool OptimizeBatch(FVectorVMCodeOptimizerContext& Context)
+	{
+		const int32 BatchedOpCount = BatchedOps.Num();
+
+		if (!BatchedOpCount)
+			return false;
+
+		for (const auto& BatchEntry : BatchedOps)
+		{
+			const uint16 AccumulatedOpCount = BatchEntry.Value.Num();
+
+			if (!AccumulatedOpCount)
+				continue;
+
+			switch (BatchEntry.Key.SrcOpType)
+			{
+			case SRCOP_RRR:
+				if (AcquireIndexConstant)
+				{
+					Context.Write<FVectorVMExecFunction>(DoRegisterKernelFixedValid<FCopyOp>);
+				}
+				else
+				{
+					Context.Write<FVectorVMExecFunction>(DoRegisterKernelVariableValid);
+				}
+				break;
+			case SRCOP_RRC:	Context.Write<FVectorVMExecFunction>(DoRegisterKernelFixedValid<FSplatOp<FConstantHandler<int32>>>); break;
+			default: check(0); break;
+			}
+
+			Context.Write(BatchEntry.Key.DataSetIndex);
+			Context.Write(BatchEntry.Key.DestIndexRegisterIdx);
+			Context.Write(AccumulatedOpCount);
+			for (const FOpValue& OpValue : BatchEntry.Value)
+			{
+				Context.Write(OpValue.SourceRegisterIndex);
+				Context.Write(OpValue.DestRegisterIdx);
+			}
+		}
+
+		return true;
+	}
+
+	bool ExtractOp(FVectorVMCodeOptimizerContext& Context)
+	{
+		FOpKey Key;
+		Key.SrcOpType = Context.BaseContext.DecodeSrcOperandTypes();
+		Key.DataSetIndex = Context.DecodeU16();
+		Key.DestIndexRegisterIdx = Context.DecodeU16();
+
+		if (Key.DestIndexRegisterIdx != WorkingRegisterIndex)
+		{
+			// if we've found an output node that is not related to the acquire index op, then just exit
+			return false;
+		}
+
+		FOpValue Value;
+		Value.SourceRegisterIndex = Context.DecodeU16();
+		Value.DestRegisterIdx = Context.DecodeU16();
+
+		TArray<FOpValue>& ExistingOps = BatchedOps.FindOrAdd(Key);
+		ExistingOps.Add(Value);
+
+		return true;
+	}
+
+private:
+	using RegisterType = VectorRegisterInt;
+	using ScalarType = int32;
+
+	uint16 DataSetCounterIndex = 0;
+	uint16 ValidTestRegisterIndex = 0;
+	uint16 WorkingRegisterIndex = 0;
+	bool AcquireIndexConstant = false;
+
+	struct FOpKey
+	{
+		uint16 DestIndexRegisterIdx;
+		uint16 DataSetIndex;
+		uint8 SrcOpType;
+	};
+
+	struct FOpValue
+	{
+		uint16 SourceRegisterIndex;
+		uint16 DestRegisterIdx;
+	};
+
+	struct FOpKeyFuncs : public TDefaultMapKeyFuncs<FOpKey, TArray<FOpValue>, false>
+	{
+		static VM_FORCEINLINE bool Matches(const FOpKey& A, const FOpKey& B)
+		{
+			return A.DestIndexRegisterIdx == B.DestIndexRegisterIdx
+				&& A.DataSetIndex == B.DataSetIndex
+				&& A.SrcOpType == B.SrcOpType;
+		}
+
+		static VM_FORCEINLINE uint32 GetKeyHash(const FOpKey& Key)
+		{
+			return HashCombine(
+				HashCombine(GetTypeHash(Key.DestIndexRegisterIdx), GetTypeHash(Key.DataSetIndex)),
+				GetTypeHash(Key.SrcOpType));
+		}
+	};
+
+	TMap<FOpKey, TArray<FOpValue>, FDefaultSetAllocator, FOpKeyFuncs> BatchedOps;
+};
+
+// look for the pattern of acquireindex followed by a bunch of outputs.
+EVectorVMOp PackedOutputOptimization(EVectorVMOp Op, FVectorVMCodeOptimizerContext& Context)
+{
+	if (!GbBatchPackVMOutput)
+	{
+		return Op;
+	}
+
+	if (Op == EVectorVMOp::acquireindex)
+	{
+		const auto RollbackState = Context.CreateCodeState();
+
+		FBatchedWriteIndexedOutput BatchedOutputOp;
+
+		BatchedOutputOp.OptimizeAcquireIndex(Context);
+
+		bool BatchValid = true;
+
+		Op = Context.BaseContext.DecodeOp();
+
+		while (BatchValid && Op == EVectorVMOp::outputdata_32bit)
+		{
+			BatchValid = BatchedOutputOp.ExtractOp(Context);
+			Op = Context.BaseContext.DecodeOp();
+		}
+
+		// if there's nothing worth optimizing here, then just revert what we've parsed
+		if (!BatchValid || !BatchedOutputOp.OptimizeBatch(Context))
+		{
+			Context.RollbackCodeState(RollbackState);
+			return EVectorVMOp::acquireindex;
+		}
+	}
+
+	return Op;
+}
+
+EVectorVMOp SafeMathOptimization(EVectorVMOp Op, FVectorVMCodeOptimizerContext& Context)
+{
+	if (!GbSafeOptimizedKernels)
+	{
+		return Op;
+	}
+
+	switch (Op)
+	{
+		case EVectorVMOp::div: FVectorKernelDivSafe::Optimize(Context); break;
+		case EVectorVMOp::rcp: FVectorKernelRcpSafe::Optimize(Context); break;
+		case EVectorVMOp::rsq: FVectorKernelRsqSafe::Optimize(Context); break;
+		case EVectorVMOp::sqrt: FVectorKernelSqrtSafe::Optimize(Context); break;
+		case EVectorVMOp::log: FVectorKernelLogSafe::Optimize(Context); break;
+		case EVectorVMOp::pow: FVectorKernelPowSafe::Optimize(Context); break;
+		default:
+			return Op;
+	}
+
+	return Context.BaseContext.DecodeOp();
+}
+
 void VectorVM::OptimizeByteCode(const uint8* ByteCode, TArray<uint8>& OptimizedCode, TArrayView<uint8> ExternalFunctionRegisterCounts)
 {
 	OptimizedCode.Empty();
@@ -2355,27 +2887,37 @@ void VectorVM::OptimizeByteCode(const uint8* ByteCode, TArray<uint8>& OptimizedC
 
 	FVectorVMCodeOptimizerContext Context(FVectorVMContext::Get(), ByteCode, OptimizedCode, ExternalFunctionRegisterCounts);
 
+	// add any optimization filters in here, useful so what we can isolate optimizations with CVars
+	FVectorVMCodeOptimizerContext::OptimizeVMFunction VMFilters[] =
+	{
+		PackedOutputOptimization,
+		SafeMathOptimization,
+	};
+
 	EVectorVMOp Op = EVectorVMOp::done;
 	do
 	{
 		Op = Context.BaseContext.DecodeOp();
 
+		for (auto Filter : VMFilters)
+			Op = Filter(Op, Context);
+
 		switch (Op)
 		{
 			case EVectorVMOp::add: FVectorKernelAdd::Optimize(Context); break;
 			case EVectorVMOp::sub: FVectorKernelSub::Optimize(Context); break;
 			case EVectorVMOp::mul: FVectorKernelMul::Optimize(Context); break;
-			case EVectorVMOp::div: GbSafeOptimizedKernels ? FVectorKernelDivSafe::Optimize(Context) : FVectorKernelDiv::Optimize(Context); break;
+			case EVectorVMOp::div: FVectorKernelDiv::Optimize(Context); break;
 			case EVectorVMOp::mad: FVectorKernelMad::Optimize(Context); break;
 			case EVectorVMOp::lerp: FVectorKernelLerp::Optimize(Context); break;
-			case EVectorVMOp::rcp: GbSafeOptimizedKernels ? FVectorKernelRcpSafe::Optimize(Context) : FVectorKernelRcp::Optimize(Context); break;
-			case EVectorVMOp::rsq: GbSafeOptimizedKernels ? FVectorKernelRsqSafe::Optimize(Context) : FVectorKernelRsq::Optimize(Context); break;
-			case EVectorVMOp::sqrt: GbSafeOptimizedKernels ? FVectorKernelSqrtSafe::Optimize(Context) : FVectorKernelSqrt::Optimize(Context); break;
+			case EVectorVMOp::rcp: FVectorKernelRcp::Optimize(Context); break;
+			case EVectorVMOp::rsq: FVectorKernelRsq::Optimize(Context); break;
+			case EVectorVMOp::sqrt: FVectorKernelSqrt::Optimize(Context); break;
 			case EVectorVMOp::neg: FVectorKernelNeg::Optimize(Context); break;
 			case EVectorVMOp::abs: FVectorKernelAbs::Optimize(Context); break;
 			case EVectorVMOp::exp: FVectorKernelExp::Optimize(Context); break;
 			case EVectorVMOp::exp2: FVectorKernelExp2::Optimize(Context); break;
-			case EVectorVMOp::log: GbSafeOptimizedKernels ? FVectorKernelLogSafe::Optimize(Context) : FVectorKernelLog::Optimize(Context); break;
+			case EVectorVMOp::log: FVectorKernelLog::Optimize(Context); break;
 			case EVectorVMOp::log2: FVectorKernelLog2::Optimize(Context); break;
 			case EVectorVMOp::sin: FVectorKernelSin::Optimize(Context); break;
 			case EVectorVMOp::cos: FVectorKernelCos::Optimize(Context); break;
@@ -2393,7 +2935,7 @@ void VectorVM::OptimizeByteCode(const uint8* ByteCode, TArray<uint8>& OptimizedC
 			case EVectorVMOp::clamp: FVectorKernelClamp::Optimize(Context); break;
 			case EVectorVMOp::min: FVectorKernelMin::Optimize(Context); break;
 			case EVectorVMOp::max: FVectorKernelMax::Optimize(Context); break;
-			case EVectorVMOp::pow: GbSafeOptimizedKernels ? FVectorKernelPowSafe::Optimize(Context) : FVectorKernelPow::Optimize(Context); break;
+			case EVectorVMOp::pow: FVectorKernelPow::Optimize(Context); break;
 			case EVectorVMOp::sign: FVectorKernelSign::Optimize(Context); break;
 			case EVectorVMOp::step: FVectorKernelStep::Optimize(Context); break;
 			case EVectorVMOp::random: FVectorKernelRandom::Optimize(Context); break;
diff --git a/Engine/Source/Runtime/VectorVM/Public/VectorVM.h b/Engine/Source/Runtime/VectorVM/Public/VectorVM.h
--- a/Engine/Source/Runtime/VectorVM/Public/VectorVM.h
+++ b/Engine/Source/Runtime/VectorVM/Public/VectorVM.h
@@ -262,6 +262,9 @@ struct FVectorVMContext : TThreadSingleton<FVectorVMContext>
 private:
 	/** Pointer to the next element in the byte code. */
 	uint8 const* RESTRICT Code;
+
+	friend struct FVectorVMCodeOptimizerContext;
+
 public:
 	/** Pointer to the constant table. */
 	uint8 const* RESTRICT ConstantTable;
@@ -302,6 +305,10 @@ public:
 
 	bool bIsParallelExecution;
 
+	int32 ValidInstanceIndexStart = INDEX_NONE;
+	int32 ValidInstanceCount = 0;
+	bool ValidInstanceUniform = false;
+
 	FVectorVMContext();
 
 	void PrepareForExec(
@@ -311,11 +318,11 @@ public:
 		void** InUserPtrTable,
 		TArrayView<FDataSetMeta> InDataSetMetaTable,
 		int32 MaxNumInstances,
-		bool bInParallelExecution
+		bool bInParallelExecution);
+
 #if STATS
-		, const TArray<TStatId>* InStatScopes
+	void SetStatScopes(const TArray<TStatId>* InStatScopes);
 #endif
-	);
 
 	void FinishExec();
 
@@ -326,6 +333,10 @@ public:
 		NumInstancesVectorFloats = (NumInstances + VECTOR_WIDTH_FLOATS - 1) / VECTOR_WIDTH_FLOATS;
 		StartInstance = InStartInstance;
 
+		ValidInstanceCount = 0;
+		ValidInstanceIndexStart = INDEX_NONE;
+		ValidInstanceUniform = false;
+
 		RandCounters.Reset();
 		RandCounters.SetNumZeroed(InNumInstances);
 	}
