Fix for vector writes crossing batch boundaries in VectorVM with batched outputs

diff --git a/Engine/Config/BaseEngine.ini b/Engine/Config/BaseEngine.ini
--- a/Engine/Config/BaseEngine.ini
+++ b/Engine/Config/BaseEngine.ini
@@ -1357,8 +1357,6 @@ con.DebugLateCheat = True
 ; disabled for now to prevent issues http://crashreport/Crashes/Show/2948912, http://crashreport/Crashes/Show/2948842
 r.DetectAndWarnOfBadDrivers = 0
 
-vm.BatchPackedVMOutput=0
-
 [SystemSettingsEditor]
 ; System settings overrides for the editor.  Ideally the editor should use the same settings as the game.
 
diff --git a/Engine/Source/Runtime/VectorVM/Private/VectorVM.cpp b/Engine/Source/Runtime/VectorVM/Private/VectorVM.cpp
--- a/Engine/Source/Runtime/VectorVM/Private/VectorVM.cpp
+++ b/Engine/Source/Runtime/VectorVM/Private/VectorVM.cpp
@@ -55,20 +55,20 @@ namespace VectorVMConstants
 	static const VectorRegisterInt RegisterShuffleMask[] =
 	{
 		MakeVectorRegisterInt(ShufMaskIgnore, ShufMaskIgnore, ShufMaskIgnore, ShufMaskIgnore), // 0000
-		MakeVectorRegisterInt(ShufMaskD, ShufMaskIgnore, ShufMaskIgnore, ShufMaskIgnore), // 0001
-		MakeVectorRegisterInt(ShufMaskC, ShufMaskIgnore, ShufMaskIgnore, ShufMaskIgnore), // 0010
-		MakeVectorRegisterInt(ShufMaskC, ShufMaskD, ShufMaskIgnore, ShufMaskIgnore), // 0011
-		MakeVectorRegisterInt(ShufMaskB, ShufMaskIgnore, ShufMaskIgnore, ShufMaskIgnore), // 0100
-		MakeVectorRegisterInt(ShufMaskB, ShufMaskD, ShufMaskIgnore, ShufMaskIgnore), // 0101
+		MakeVectorRegisterInt(ShufMaskA, ShufMaskIgnore, ShufMaskIgnore, ShufMaskIgnore), // 0001
+		MakeVectorRegisterInt(ShufMaskB, ShufMaskIgnore, ShufMaskIgnore, ShufMaskIgnore), // 0010
+		MakeVectorRegisterInt(ShufMaskA, ShufMaskB, ShufMaskIgnore, ShufMaskIgnore), // 0011
+		MakeVectorRegisterInt(ShufMaskC, ShufMaskIgnore, ShufMaskIgnore, ShufMaskIgnore), // 0100
+		MakeVectorRegisterInt(ShufMaskA, ShufMaskC, ShufMaskIgnore, ShufMaskIgnore), // 0101
 		MakeVectorRegisterInt(ShufMaskB, ShufMaskC, ShufMaskIgnore, ShufMaskIgnore), // 0110
-		MakeVectorRegisterInt(ShufMaskB, ShufMaskC, ShufMaskD, ShufMaskIgnore), // 0111
-		MakeVectorRegisterInt(ShufMaskA, ShufMaskIgnore, ShufMaskIgnore, ShufMaskIgnore), // 1000
+		MakeVectorRegisterInt(ShufMaskA, ShufMaskB, ShufMaskC, ShufMaskIgnore), // 0111
+		MakeVectorRegisterInt(ShufMaskD, ShufMaskIgnore, ShufMaskIgnore, ShufMaskIgnore), // 1000
 		MakeVectorRegisterInt(ShufMaskA, ShufMaskD, ShufMaskIgnore, ShufMaskIgnore), // 1001
-		MakeVectorRegisterInt(ShufMaskA, ShufMaskC, ShufMaskIgnore, ShufMaskIgnore), // 1010
-		MakeVectorRegisterInt(ShufMaskA, ShufMaskC, ShufMaskD, ShufMaskIgnore), // 1011
-		MakeVectorRegisterInt(ShufMaskA, ShufMaskB, ShufMaskIgnore, ShufMaskIgnore), // 1100
-		MakeVectorRegisterInt(ShufMaskA, ShufMaskB, ShufMaskD, ShufMaskIgnore), // 1101
-		MakeVectorRegisterInt(ShufMaskA, ShufMaskB, ShufMaskC, ShufMaskIgnore), // 1110
+		MakeVectorRegisterInt(ShufMaskB, ShufMaskD, ShufMaskIgnore, ShufMaskIgnore), // 1010
+		MakeVectorRegisterInt(ShufMaskA, ShufMaskB, ShufMaskD, ShufMaskIgnore), // 1011
+		MakeVectorRegisterInt(ShufMaskC, ShufMaskD, ShufMaskIgnore, ShufMaskIgnore), // 1100
+		MakeVectorRegisterInt(ShufMaskA, ShufMaskC, ShufMaskD, ShufMaskIgnore), // 1101
+		MakeVectorRegisterInt(ShufMaskB, ShufMaskC, ShufMaskD, ShufMaskIgnore), // 1110
 		MakeVectorRegisterInt(ShufMaskA, ShufMaskB, ShufMaskC, ShufMaskD), // 1111
 	};
 
@@ -2564,9 +2564,6 @@ struct FBatchedWriteIndexedOutput
 		FDataSetMeta& DataSetMeta = Context.GetDataSetMeta(DataSetIndex);
 		
 		const int8* DestIndexReg = reinterpret_cast<const int8*>(Context.GetTempRegister(DestIndexRegisterIdx));
-		const int32 LoopCount = FMath::DivideAndRoundUp(Context.NumInstances, VECTOR_WIDTH_FLOATS);
-
-		uint16 OpIt = 0;
 
 		//
 		// VectorIntStore(		- unaligned writes of 16 bytes to our Destination; note that this maneuver requires us to have
@@ -2575,52 +2572,57 @@ struct FBatchedWriteIndexedOutput
 		//    Source,			- source data
 		//    ShuffleMask),		- result of the VectorMaskBits done in the acquireindex, int8/VectorRegister of input
 		//  Destination);
-
-		constexpr int32 OpsPerLoop = 4;
-		for (OpIt = 0; (OpIt + OpsPerLoop) < AccumulatedOpCount; OpIt += OpsPerLoop)
+		for (uint16 OpIt = 0; OpIt < AccumulatedOpCount; ++OpIt)
 		{
-			const RegisterType* Source0 = FRegisterHandler<RegisterType>(Context).GetDest();
-			int32* DestReg0 = Context.GetOutputRegister<int32>(DataSetIndex, Context.DecodeU16()) + Context.ValidInstanceIndexStart;
-
-			const RegisterType* Source1 = FRegisterHandler<RegisterType>(Context).GetDest();
-			int32* DestReg1 = Context.GetOutputRegister<int32>(DataSetIndex, Context.DecodeU16()) + Context.ValidInstanceIndexStart;
-
-			const RegisterType* Source2 = FRegisterHandler<RegisterType>(Context).GetDest();
-			int32* DestReg2 = Context.GetOutputRegister<int32>(DataSetIndex, Context.DecodeU16()) + Context.ValidInstanceIndexStart;
+			const RegisterType* Source = FRegisterHandler<RegisterType>(Context).GetDest();
+			int32* DestReg = Context.GetOutputRegister<int32>(DataSetIndex, Context.DecodeU16()) + Context.ValidInstanceIndexStart;
 
-			const RegisterType* Source3 = FRegisterHandler<RegisterType>(Context).GetDest();
-			int32* DestReg3 = Context.GetOutputRegister<int32>(DataSetIndex, Context.DecodeU16()) + Context.ValidInstanceIndexStart;
+			// the number of instances that we're expecting to write.  it is important that we keep track of it because when we
+			// get down to the end we need to switch from the shuffled approach to a scalar approach so that we don't
+			// overwrite the indexed output that another parallel context might have written to
+			int32 WritesRemaining = Context.ValidInstanceCount;
+			int32 SourceIt = 0;
 
-			for (int32 LoopIt = 0; LoopIt < LoopCount; ++LoopIt)
+			// vector shuffle path writes 4 at a time (though the trailing elements may not be valid) until we have to move over
+			// to the scalar version for fear of overwriting our neighbors
+			while (WritesRemaining >= VECTOR_WIDTH_FLOATS)
 			{
-				const int8 ShuffleMask = DestIndexReg[LoopIt];
+				check(SourceIt * VECTOR_WIDTH_FLOATS < Context.NumInstances);
+
+				const int8 ShuffleMask = DestIndexReg[SourceIt];
 				const int8 AdvanceCount = FMath::CountBits(ShuffleMask);
 
-				VectorIntStore(VectorIntShuffle(Source0[LoopIt], VectorVMConstants::RegisterShuffleMask[ShuffleMask]), DestReg0);
-				VectorIntStore(VectorIntShuffle(Source1[LoopIt], VectorVMConstants::RegisterShuffleMask[ShuffleMask]), DestReg1);
-				VectorIntStore(VectorIntShuffle(Source2[LoopIt], VectorVMConstants::RegisterShuffleMask[ShuffleMask]), DestReg2);
-				VectorIntStore(VectorIntShuffle(Source3[LoopIt], VectorVMConstants::RegisterShuffleMask[ShuffleMask]), DestReg3);
+				VectorIntStore(VectorIntShuffle(Source[SourceIt], VectorVMConstants::RegisterShuffleMask[ShuffleMask]), DestReg);
 
-				DestReg0 += AdvanceCount;
-				DestReg1 += AdvanceCount;
-				DestReg2 += AdvanceCount;
-				DestReg3 += AdvanceCount;
-			}
-		}
+				DestReg += AdvanceCount;
+				WritesRemaining -= AdvanceCount;
 
-		for (; OpIt < AccumulatedOpCount; ++OpIt)
-		{
-			const RegisterType* Source = FRegisterHandler<RegisterType>(Context).GetDest();
-			int32* DestReg0 = Context.GetOutputRegister<int32>(DataSetIndex, Context.DecodeU16()) + Context.ValidInstanceIndexStart;
+				++SourceIt;
+			}
 
-			for (int32 LoopIt = 0; LoopIt < LoopCount; ++LoopIt)
+			// scalar path that will read 4 values and write one at a time to the output based on the valid mask
+			while (WritesRemaining)
 			{
-				const int8 ShuffleMask = DestIndexReg[LoopIt];
+				const int8 ShuffleMask = DestIndexReg[SourceIt];
 				const int8 AdvanceCount = FMath::CountBits(ShuffleMask);
+				if (AdvanceCount)
+				{
+					int32 RawSourceData[VECTOR_WIDTH_FLOATS];
 
-				VectorIntStore(VectorIntShuffle(Source[LoopIt], VectorVMConstants::RegisterShuffleMask[ShuffleMask]), DestReg0);
+					VectorIntStore(Source[SourceIt], RawSourceData);
 
-				DestReg0 += AdvanceCount;
+					for (int32 ScalarIt = 0; ScalarIt < 4; ++ScalarIt)
+					{
+						if (!!(ShuffleMask & (1 << ScalarIt)))
+						{
+							*DestReg = RawSourceData[ScalarIt];
+							++DestReg;
+						}
+					}
+
+					WritesRemaining -= AdvanceCount;
+				}
+				++SourceIt;
 			}
 		}
 	}
