- Ensure we don't accidentally read from the constant table when optimizing stat calls inside the VectorVM
- If the compiler chooses to fold both the Optimize & regular VM function calls together we will read garbage

#rnx

diff --git a/Engine/Source/Runtime/VectorVM/Private/VectorVM.cpp b/Engine/Source/Runtime/VectorVM/Private/VectorVM.cpp
--- a/Engine/Source/Runtime/VectorVM/Private/VectorVM.cpp
+++ b/Engine/Source/Runtime/VectorVM/Private/VectorVM.cpp
@@ -292,6 +292,11 @@ struct FConstantHandlerBase
 	{
 		Context.Write(Context.DecodeU16());
 	}
+
+	static void OptimizeSkip(FVectorVMCodeOptimizerContext& Context)
+	{
+		Context.DecodeU16();
+	}
 };
 
 template<typename T>
@@ -1070,7 +1075,7 @@ struct FVectorKernelEnterStatScope
 		FConstantHandler<int32>::Optimize(Context);
 #else
 		// just skip the op if we don't have stats enabled
-		FConstantHandler<int32>(Context.BaseContext);
+		FConstantHandler<int32>::OptimizeSkip(Context);
 #endif
 	}
 
