Fixed incorrect ID to index mapping on CPU when an emitter spawns additional particles.

The UpdateID CPU kernel didn't take into account the instance start offset, so the spawn step mapped new IDs to indices starting at 0. If the emitter already contained particles, this caused the ID to index mapping to be incorrect in the frame when the spawn happened, which caused further problems if another emitter tried to read the particles in that frame.

diff --git a/Engine/Plugins/FX/Niagara/Source/Niagara/Private/NiagaraDataSet.cpp b/Engine/Plugins/FX/Niagara/Source/Niagara/Private/NiagaraDataSet.cpp
--- a/Engine/Plugins/FX/Niagara/Source/Niagara/Private/NiagaraDataSet.cpp
+++ b/Engine/Plugins/FX/Niagara/Source/Niagara/Private/NiagaraDataSet.cpp
@@ -226,7 +226,6 @@ void FNiagaraDataSet::Allocate(int32 NumInstances, bool bMaintainExisting)
 
 		int32 RequiredIDs = FMath::Max(NumInstances, NumUsedIDs);
 		int32 ExistingNumIDs = CurrentIDTable.Num();
-		int32 NumNewIDs = RequiredIDs - ExistingNumIDs;
 
 		//////////////////////////////////////////////////////////////////////////
 		//TODO: We should replace this with a lock free list that uses just a single table with RequiredIDs elements.
@@ -247,13 +246,13 @@ void FNiagaraDataSet::Allocate(int32 NumInstances, bool bMaintainExisting)
 #endif
 
 			//Free ID Table must always be at least as large as the data buffer + it's current size in the case all particles die this frame.
-			FreeIDsTable.AddUninitialized(NumNewIDs);
+			FreeIDsTable.AddUninitialized(NewNumIds);
 
 			// The spawned IDs table must be as large as the free IDs table, in case we allocate all of them this tick.
 			SpawnedIDsTable.Reserve(FreeIDsTable.Num());
 
 			//Free table should always have enough room for these new IDs.
-			check(NumFreeIDs + NumNewIDs <= FreeIDsTable.Num());
+			check(NumFreeIDs + NewNumIds <= FreeIDsTable.Num());
 
 			//ID Table grows so add any new IDs to the free array. Add in reverse order to maintain a continuous increasing allocation when popping.
 			for (int32 NewFreeID = RequiredIDs - 1; NewFreeID >= ExistingNumIDs; --NewFreeID)
diff --git a/Engine/Source/Runtime/VectorVM/Private/VectorVM.cpp b/Engine/Source/Runtime/VectorVM/Private/VectorVM.cpp
--- a/Engine/Source/Runtime/VectorVM/Private/VectorVM.cpp
+++ b/Engine/Source/Runtime/VectorVM/Private/VectorVM.cpp
@@ -1482,7 +1482,8 @@ struct FScalarKernelUpdateID
 		const TArrayView<FDataSetMeta> MetaTable = Context.DataSetMetaTable;
 
 		TArray<int32>&RESTRICT IDTable = *MetaTable[DataSetIndex].IDTable;
-		const int32 InstanceOffset = MetaTable[DataSetIndex].InstanceOffset + Context.StartInstance;
+		const int32 InstanceOffset = MetaTable[DataSetIndex].InstanceOffset;
+		const int32 AbsoluteStartInstance = InstanceOffset + Context.StartInstance;
 
 		const int32*RESTRICT IDRegister = (int32*)(Context.GetTempRegister(InstanceIDRegisterIndex));
 		const int32*RESTRICT IndexRegister = (int32*)(Context.GetTempRegister(InstanceIndexRegisterIndex));
@@ -1490,11 +1491,14 @@ struct FScalarKernelUpdateID
 		FDataSetThreadLocalTempData& DataSetTempData = Context.ThreadLocalTempData[DataSetIndex];
 
 		TArray<int32>&RESTRICT IDsToFree = DataSetTempData.IDsToFree;
-		check(IDTable.Num() >= InstanceOffset + Context.NumInstances);
+		check(IDTable.Num() >= AbsoluteStartInstance + Context.NumInstances);
 		for (int32 i = 0; i < Context.NumInstances; ++i)
 		{
 			int32 InstanceId = IDRegister[i];
-			int32 Index = IndexRegister[i];
+			// The index passed into this function is the same as that given to the OutputData*() functions (FScalarKernelWriteOutputIndexed kernel).
+			// That value is local to the execution step (update or spawn), so we must offset it by the start instance number for the step, just like
+			// GetOutputRegister() does.
+			int32 Index = IndexRegister[i] + InstanceOffset;
 
 			if (Index == INDEX_NONE)
 			{
