Fixed particle IDs never being deallocated on CPU (bug introduced in 11290626).

diff --git a/Engine/Plugins/FX/Niagara/Source/Niagara/Private/NiagaraDataSet.cpp b/Engine/Plugins/FX/Niagara/Source/Niagara/Private/NiagaraDataSet.cpp
--- a/Engine/Plugins/FX/Niagara/Source/Niagara/Private/NiagaraDataSet.cpp
+++ b/Engine/Plugins/FX/Niagara/Source/Niagara/Private/NiagaraDataSet.cpp
@@ -240,9 +240,6 @@ void FNiagaraDataSet::Allocate(int32 NumInstances, bool bMaintainExisting)
 			//Free ID Table must always be at least as large as the data buffer + it's current size in the case all particles die this frame.
 			FreeIDsTable.AddUninitialized(NewNumIds);
 
-			// The spawned IDs table must be as large as the free IDs table, in case we allocate all of them this tick.
-			SpawnedIDsTable.Reserve(FreeIDsTable.Num());
-
 			//Free table should always have enough room for these new IDs.
 			check(NumFreeIDs + NewNumIds <= FreeIDsTable.Num());
 
@@ -285,6 +282,9 @@ void FNiagaraDataSet::Allocate(int32 NumInstances, bool bMaintainExisting)
 			RequiredIDs = ExistingNumIDs;
 		}
 
+		// We know that we can't spawn more than NumFreeIDs particles, so we can pre-allocate SpawnedIDsTable here, to avoid allocations during execution.
+		SpawnedIDsTable.Reserve(NumFreeIDs);
+
 		// We need to clear the ID to index table to -1 so we don't have stale entries for particles which died in the previous
 		// frame (when the results were written to another buffer). All the entries which are in use will be filled in by the script.
 		DestinationIDTable.SetNumUninitialized(RequiredIDs);
diff --git a/Engine/Source/Runtime/VectorVM/Private/VectorVM.cpp b/Engine/Source/Runtime/VectorVM/Private/VectorVM.cpp
--- a/Engine/Source/Runtime/VectorVM/Private/VectorVM.cpp
+++ b/Engine/Source/Runtime/VectorVM/Private/VectorVM.cpp
@@ -1495,10 +1495,7 @@ struct FScalarKernelUpdateID
 		for (int32 i = 0; i < Context.NumInstances; ++i)
 		{
 			int32 InstanceId = IDRegister[i];
-			// The index passed into this function is the same as that given to the OutputData*() functions (FScalarKernelWriteOutputIndexed kernel).
-			// That value is local to the execution step (update or spawn), so we must offset it by the start instance number for the step, just like
-			// GetOutputRegister() does.
-			int32 Index = IndexRegister[i] + InstanceOffset;
+			int32 Index = IndexRegister[i];
 
 			if (Index == INDEX_NONE)
 			{
@@ -1508,8 +1505,11 @@ struct FScalarKernelUpdateID
 			}
 			else
 			{
-				//Update the actual index for this ID. No thread safety is needed as this ID slot can only ever be written by this instance and so a single thread.
-				IDTable[InstanceId] = Index;
+				// Update the actual index for this ID. No thread safety is needed as this ID slot can only ever be written by this instance and so a single thread.
+				// The index passed into this function is the same as that given to the OutputData*() functions (FScalarKernelWriteOutputIndexed kernel).
+				// That value is local to the execution step (update or spawn), so we must offset it by the start instance number for the step, just like
+				// GetOutputRegister() does.
+				IDTable[InstanceId] = Index + InstanceOffset;
 
 				//Update thread local max ID seen. We push this to the real value at the end of execution.
 				DataSetTempData.MaxID = FMath::Max(DataSetTempData.MaxID, InstanceId);
