Adding a check if we are overflowing the specified constant buffer table

diff --git a/Engine/Source/Runtime/VectorVM/Public/VectorVM.h b/Engine/Source/Runtime/VectorVM/Public/VectorVM.h
--- a/Engine/Source/Runtime/VectorVM/Public/VectorVM.h
+++ b/Engine/Source/Runtime/VectorVM/Public/VectorVM.h
@@ -417,6 +417,7 @@ public:
 			++TableIndex;
 		}
 
+		check(TableIndex < ConstantTableCount);
 		check(Offset < ConstantTableSizes[TableIndex]);
 		return reinterpret_cast<const T*>(ConstantTable[TableIndex] + Offset);
 	}
