FNM - Added Compression option for Niagara Systems. Thge Option will try to compress some of the parameters to use half precission, which can be a performance benefit on the GPU and sometimes also on 

diff --git a/Engine/Source/Runtime/VectorVM/Private/VectorVM.cpp b/Engine/Source/Runtime/VectorVM/Private/VectorVM.cpp
--- a/Engine/Source/Runtime/VectorVM/Private/VectorVM.cpp
+++ b/Engine/Source/Runtime/VectorVM/Private/VectorVM.cpp
@@ -74,17 +74,6 @@ namespace VectorVMConstants
 		MakeVectorRegisterInt(ShufMaskB, ShufMaskC, ShufMaskD, ShufMaskIgnore), // 1110
 		MakeVectorRegisterInt(ShufMaskA, ShufMaskB, ShufMaskC, ShufMaskD), // 1111
 	};
-
-	constexpr uint32 cOne = 0xFFFFFFFFU;
-	constexpr uint32 cZero = 0x00000000U;
-	static const VectorRegister RemainderMask[] =
-	{
-		MakeVectorRegister(cZero, cZero, cZero, cZero), // 0 remaining
-		MakeVectorRegister(cOne, cZero, cZero, cZero), // 1 remaining
-		MakeVectorRegister(cOne, cOne, cZero, cZero), // 2 remaining
-		MakeVectorRegister(cOne, cOne, cOne, cZero), // 3 remaining
-		MakeVectorRegister(cOne, cOne, cOne, cOne), // 4 remaining
-	};
 };
 
 // helper function wrapping the SSE3 shuffle operation.  Currently implemented for PS4/XB1/Neon, the
@@ -102,7 +91,8 @@ FORCEINLINE VectorRegisterInt VectorIntShuffle(const VectorRegisterInt& Vec, con
 {
 	uint8x8x2_t VecSplit = { { vget_low_u8(Vec), vget_high_u8(Vec) } };
 	return vcombine_u8(vtbl2_u8(VecSplit, vget_low_u8(Mask)), vtbl2_u8(VecSplit, vget_high_u8(Mask)));
-	}
+}
+
 #else
 FORCEINLINE VectorRegisterInt VectorIntShuffle(const VectorRegisterInt& Vec, const VectorRegisterInt& Mask)
 {
@@ -209,6 +199,22 @@ static FAutoConsoleVariableRef CVarbSafeOptimizedKernels(
 	ECVF_Default
 );
 
+static int32 GbBatchVMInput = 0;
+static FAutoConsoleVariableRef CVarBatchVMInput(
+	TEXT("vm.BatchVMInput"),
+	GbBatchVMInput,
+	TEXT("If > 0 input elements will be batched.\n"),
+	ECVF_Default
+);
+
+static int32 GbBatchVMOutput = 0;
+static FAutoConsoleVariableRef CVarBatchVMOutput(
+	TEXT("vm.BatchVMOutput"),
+	GbBatchVMOutput,
+	TEXT("If > 0 output elements will be batched.\n"),
+	ECVF_Default
+);
+
 static int32 GbBatchPackVMOutput = 1;
 static FAutoConsoleVariableRef CVarbBatchPackVMOutput(
 	TEXT("vm.BatchPackedVMOutput"),
@@ -1521,7 +1527,7 @@ struct FScalarKernelUpdateID
 };
 
 /** Special kernel for reading from the main input dataset. */
-template<typename T>
+template<typename SourceType, int TypeOffset>
 struct FVectorKernelReadInput
 {
 	static void Optimize(FVectorVMCodeOptimizerContext& Context)
@@ -1534,33 +1540,73 @@ struct FVectorKernelReadInput
 
 	static VM_FORCEINLINE void Exec(FVectorVMContext& Context)
 	{
-		static const int32 InstancesPerVector = sizeof(VectorRegister) / sizeof(T);
+		static const int32 InstancesPerVector = sizeof(VectorRegister) / sizeof(SourceType);
 
 		const int32 DataSetIndex = Context.DecodeU16();
 		const int32 InputRegisterIdx = Context.DecodeU16();
 		const int32 DestRegisterIdx = Context.DecodeU16();
-		const int32 Loops = Context.GetNumLoops<InstancesPerVector>();
+		int32 Loops = Context.GetNumLoops<InstancesPerVector>();
 
 		VectorRegister* DestReg = (VectorRegister*)(Context.GetTempRegister(DestRegisterIdx));
-		VectorRegister* InputReg = (VectorRegister*)(Context.GetInputRegister<T>(DataSetIndex, InputRegisterIdx) + Context.GetStartInstance());
+		VectorRegister* InputReg = (VectorRegister*)(Context.GetInputRegister<SourceType, TypeOffset>(DataSetIndex, InputRegisterIdx) + Context.GetStartInstance());
 
 		//TODO: We can actually do some scalar loads into the first and final vectors to get around alignment issues and then use the aligned load for all others.
-		for (int32 i = 0; i < Loops; ++i)
+		while(Loops > 0)
 		{
 			*DestReg = VectorLoad(InputReg);
 			++DestReg;
 			++InputReg;
+			Loops--;
 		}
 	}
 };
 
+template<>
+struct FVectorKernelReadInput<FFloat16, 2>
+{
+	static void Optimize(FVectorVMCodeOptimizerContext& Context)
+	{
+		Context.Write<FVectorVMExecFunction>(Exec);
+		Context.Write(Context.DecodeU16());	// DataSetIndex
+		Context.Write(Context.DecodeU16());	// InputRegisterIdx
+		Context.Write(Context.DecodeU16());	// DestRegisterIdx
+	}
+
+	static VM_FORCEINLINE void Exec(FVectorVMContext& Context)
+	{
+		const int32 DataSetIndex = Context.DecodeU16();
+		const int32 InputRegisterIdx = Context.DecodeU16();
+		const int32 DestRegisterIdx = Context.DecodeU16();
+		int32 Loops = Context.GetNumLoops<4>();
+
+		float* DestReg = (float*)(Context.GetTempRegister(DestRegisterIdx));
+		uint16* InputReg = (uint16*)(Context.GetInputRegister<FFloat16, 2>(DataSetIndex, InputRegisterIdx) + Context.GetStartInstance());
+
+		//TODO: We can actually do some scalar loads into the first and final vectors to get around alignment issues and then use the aligned load for all others.
+		while (Loops > 1)
+		{
+			FPlatformMath::WideVectorLoadHalf(DestReg, InputReg);
+			DestReg += 8;
+			InputReg += 8;
+			Loops -= 2;
+		}
+
+		while (Loops > 0)
+		{
+			FPlatformMath::VectorLoadHalf(DestReg, InputReg);
+			DestReg += 4;
+			InputReg += 4;
+			Loops -= 1;
+		}
+	}
+};
 
 
 /** Special kernel for reading from an input dataset; non-advancing (reads same instance everytime). 
  *  this kernel splats the X component of the source register to all 4 dest components; it's meant to
  *	use scalar data sets as the source (e.g. events)
  */
-template<typename T>
+template<typename T, int TypeOffset>
 struct FVectorKernelReadInputNoAdvance
 {
 	static void Optimize(FVectorVMCodeOptimizerContext& Context)
@@ -1581,7 +1627,7 @@ struct FVectorKernelReadInputNoAdvance
 		const int32 Loops = Context.GetNumLoops<InstancesPerVector>();
 
 		VectorRegister* DestReg = (VectorRegister*)(Context.GetTempRegister(DestRegisterIdx));
-		VectorRegister* InputReg = (VectorRegister*)(Context.GetInputRegister<T>(DataSetIndex, InputRegisterIdx));
+		VectorRegister* InputReg = (VectorRegister*)(Context.GetInputRegister<T, TypeOffset>(DataSetIndex, InputRegisterIdx));
 
 		//TODO: We can actually do some scalar loads into the first and final vectors to get around alignment issues and then use the aligned load for all others.
 		for (int32 i = 0; i < Loops; ++i)
@@ -1592,9 +1638,36 @@ struct FVectorKernelReadInputNoAdvance
 	}
 };
 
+template<>
+struct FVectorKernelReadInputNoAdvance<FFloat16, 2>
+{
+	static void Optimize(FVectorVMCodeOptimizerContext& Context)
+	{
+		Context.Write<FVectorVMExecFunction>(Exec);
+		Context.Write(Context.DecodeU16());	// DataSetIndex
+		Context.Write(Context.DecodeU16());	// InputRegisterIdx
+		Context.Write(Context.DecodeU16());	// DestRegisterIdx
+	}
 
+	static VM_FORCEINLINE void Exec(FVectorVMContext& Context)
+	{
+		const int32 DataSetIndex = Context.DecodeU16();
+		const int32 InputRegisterIdx = Context.DecodeU16();
+		const int32 DestRegisterIdx = Context.DecodeU16();
+		const int32 Loops = Context.GetNumLoops<4>();
 
+		VectorRegister* DestReg = (VectorRegister*)(Context.GetTempRegister(DestRegisterIdx));
+		uint16* InputReg = (uint16*)(Context.GetInputRegister<FFloat16, 2>(DataSetIndex, InputRegisterIdx));
 
+		//TODO: We can actually do some scalar loads into the first and final vectors to get around alignment issues and then use the aligned load for all others.
+		for (int32 i = 0; i < Loops; ++i)
+		{
+			float flt = FPlatformMath::LoadHalf(InputReg);
+			*DestReg = VectorLoadFloat1(&flt);
+			++DestReg;
+		}
+	}
+};
 
 //TODO - Should be straight forwards to follow the input with a mix of the outputs direct indexing
 /** Special kernel for reading an specific location in an input register. */
@@ -1619,7 +1692,7 @@ struct FVectorKernelReadInputNoAdvance
 // };
 
 /** Special kernel for writing to a specific output register. */
-template<typename T>
+template<typename SourceType, typename DestType, int TypeOffset>
 struct FScalarKernelWriteOutputIndexed
 {
 	static VM_FORCEINLINE void Optimize(FVectorVMCodeOptimizerContext& Context)
@@ -1627,8 +1700,8 @@ struct FScalarKernelWriteOutputIndexed
 		const uint32 SrcOpTypes = Context.BaseContext.DecodeSrcOperandTypes();
 		switch (SrcOpTypes)
 		{
-			case SRCOP_RRR: Context.Write<FVectorVMExecFunction>(DoKernel<FRegisterHandler<T>>); break;
-			case SRCOP_RRC:	Context.Write<FVectorVMExecFunction>(DoKernel<FConstantHandler<T>>); break;
+			case SRCOP_RRR: Context.Write<FVectorVMExecFunction>(DoKernel<FRegisterHandler<SourceType>>); break;
+			case SRCOP_RRC:	Context.Write<FVectorVMExecFunction>(DoKernel<FConstantHandler<SourceType>>); break;
 			default: check(0); break;
 		};
 
@@ -1643,8 +1716,8 @@ struct FScalarKernelWriteOutputIndexed
 		const uint32 SrcOpTypes = Context.DecodeSrcOperandTypes();
 		switch (SrcOpTypes)
 		{
-		case SRCOP_RRR: DoKernel<FRegisterHandler<T>>(Context); break;
-		case SRCOP_RRC:	DoKernel<FConstantHandler<T>>(Context); break;
+		case SRCOP_RRR: DoKernel<FRegisterHandler<SourceType>>(Context); break;
+		case SRCOP_RRC:	DoKernel<FConstantHandler<SourceType>>(Context); break;
 		default: check(0); break;
 		};
 	}
@@ -1655,14 +1728,15 @@ struct FScalarKernelWriteOutputIndexed
 		const int32 DataSetIndex = Context.DecodeU16();
 
 		const int32 DestIndexRegisterIdx = Context.DecodeU16();
-		T* DestIndexReg = (T*)(Context.GetTempRegister(DestIndexRegisterIdx));
+		int32* RESTRICT DestIndexReg = (int32*)(Context.GetTempRegister(DestIndexRegisterIdx));
 
 		DataHandlerType DataHandler(Context);
 
 		const int32 DestRegisterIdx = Context.DecodeU16();
-		T* DestReg = Context.GetOutputRegister<T>(DataSetIndex, DestRegisterIdx);
+		DestType* RESTRICT DestReg = Context.GetOutputRegister<DestType, TypeOffset>(DataSetIndex, DestRegisterIdx);
 
-		for (int32 i = 0; i < Context.GetNumInstances(); ++i)
+		int NumInstances = Context.GetNumInstances();
+		for (int32 i = 0; i < NumInstances; ++i)
 		{
 			int32 DestIndex = *DestIndexReg;
 			if (DestIndex != INDEX_NONE)
@@ -1677,6 +1751,65 @@ struct FScalarKernelWriteOutputIndexed
 	}
 };
 
+template<>
+struct FScalarKernelWriteOutputIndexed<float, FFloat16, 2>
+{
+	static VM_FORCEINLINE void Optimize(FVectorVMCodeOptimizerContext& Context)
+	{
+		const uint32 SrcOpTypes = Context.BaseContext.DecodeSrcOperandTypes();
+		switch (SrcOpTypes)
+		{
+		case SRCOP_RRR: Context.Write<FVectorVMExecFunction>(DoKernel<FRegisterHandler<float>>); break;
+		case SRCOP_RRC:	Context.Write<FVectorVMExecFunction>(DoKernel<FConstantHandler<float>>); break;
+		default: check(0); break;
+		};
+
+		Context.Write(Context.DecodeU16());		// DataSetIndex
+		Context.Write(Context.DecodeU16());		// DestIndexRegisterIdx
+		Context.Write(Context.DecodeU16());		// DataHandlerType
+		Context.Write(Context.DecodeU16());		// DestRegisterIdx
+	}
+
+	static VM_FORCEINLINE void Exec(FVectorVMContext& Context)
+	{
+		const uint32 SrcOpTypes = Context.DecodeSrcOperandTypes();
+		switch (SrcOpTypes)
+		{
+		case SRCOP_RRR: DoKernel<FRegisterHandler<float>>(Context); break;
+		case SRCOP_RRC:	DoKernel<FConstantHandler<float>>(Context); break;
+		default: check(0); break;
+		};
+	}
+
+	template<typename DataHandlerType>
+	static VM_FORCEINLINE void DoKernel(FVectorVMContext& Context)
+	{
+		const int32 DataSetIndex = Context.DecodeU16();
+
+		const int32 DestIndexRegisterIdx = Context.DecodeU16();
+		int32* RESTRICT DestIndexReg = (int32*)(Context.GetTempRegister(DestIndexRegisterIdx));
+
+		DataHandlerType DataHandler(Context);
+
+		const int32 DestRegisterIdx = Context.DecodeU16();
+		uint16* RESTRICT DestReg = (uint16*)Context.GetOutputRegister<FFloat16, 2>(DataSetIndex, DestRegisterIdx);
+
+		int NumInstances = Context.GetNumInstances();
+		for (int32 i = 0; i < NumInstances; ++i)
+		{
+			int32 DestIndex = *DestIndexReg;
+			if (DestIndex != INDEX_NONE)
+			{
+				FPlatformMath::StoreHalf(&DestReg[DestIndex], DataHandler.Get());
+			}
+
+			++DestIndexReg;
+			DataHandler.Advance();
+			//We don't increment the dest as we index into it directly.
+		}
+	}
+};
+
 struct FDataSetCounterHandler
 {
 	int32* Counter;
@@ -2418,9 +2551,15 @@ void VectorVM::Exec(
 						case EVectorVMOp::i2b: FVectorKernelIntToBool::Exec(Context); break;
 						case EVectorVMOp::b2i: FVectorKernelBoolToInt::Exec(Context); break;
 
-						case EVectorVMOp::outputdata_32bit:	FScalarKernelWriteOutputIndexed<int32>::Exec(Context);	break;
-						case EVectorVMOp::inputdata_32bit: FVectorKernelReadInput<int32>::Exec(Context); break;
-						case EVectorVMOp::inputdata_noadvance_32bit: FVectorKernelReadInputNoAdvance<int32>::Exec(Context); break;
+						case EVectorVMOp::outputdata_half:	FScalarKernelWriteOutputIndexed<float, FFloat16, 2>::Exec(Context);	break;
+						case EVectorVMOp::inputdata_half: FVectorKernelReadInput<FFloat16, 2>::Exec(Context); break;
+						case EVectorVMOp::outputdata_int32:	FScalarKernelWriteOutputIndexed<int32, int32, 1>::Exec(Context);	break;
+						case EVectorVMOp::inputdata_int32: FVectorKernelReadInput<int32, 1>::Exec(Context); break;
+						case EVectorVMOp::outputdata_float:	FScalarKernelWriteOutputIndexed<float, float, 0>::Exec(Context);	break;
+						case EVectorVMOp::inputdata_float: FVectorKernelReadInput<float, 0>::Exec(Context); break;
+						case EVectorVMOp::inputdata_noadvance_int32: FVectorKernelReadInputNoAdvance<int32, 1>::Exec(Context); break;
+						case EVectorVMOp::inputdata_noadvance_float: FVectorKernelReadInputNoAdvance<float, 0>::Exec(Context); break;
+						case EVectorVMOp::inputdata_noadvance_half: FVectorKernelReadInputNoAdvance<FFloat16, 2>::Exec(Context); break;
 						case EVectorVMOp::acquireindex:	FScalarKernelAcquireCounterIndex::Exec(Context); break;
 						case EVectorVMOp::external_func_call: FKernelExternalFunctionCall::Exec(Context); break;
 
@@ -2488,180 +2627,90 @@ FString VectorVM::GetOperandLocationName(EVectorVMOperandLocation Location)
 }
 #endif
 
-// local implementation of VectorIntShuffle for neon/directx/
-
 
-// Optimization managed by GbBatchPackVMOutput via PackedOutputOptimization()
-// Looks for the common pattern of an acquireindex op followed by a number of associated outputdata_32bit ops.  The
-// stock operation is to write an index into a temporary register, and then have the different outputs streams
-// write into the indexed location.  This optimization does a number of things:
-// -first we check if 'validity' is uniform or not, if it is we can have a fast path of both figuring out how many
-// indices we need, as well as how to write the output (if we find that they are all invalid, then we don't need to do anything!)
-// -if we need to evaluate the validity of each element we quickly count up the number (with vector intrinsics) and
-// grab a block of the indices (rather than one at a time)
-// -rather than storing the indices to use, we store a int8 mask which indicates a valid flag for each of the next 4 samples
-// -outputs are then written to depending on their source and their frequency:
-//		-uniform sources will be splatted to all valid entries
-//		-variable sources will be packed into the available slots
-struct FBatchedWriteIndexedOutput
+void ExecBatchedOutput(FVectorVMContext& Context)
 {
-	// functor for copying a source register to an output register
-	struct FCopyOp
+	while (true)
 	{
-		void VM_FORCEINLINE operator()(FVectorVMContext& Context, uint16 DataSetIndex)
+		FVectorVMExecFunction ExecFunction = reinterpret_cast<FVectorVMExecFunction>(Context.DecodePtr());
+		if (ExecFunction == nullptr)
 		{
-			FRegisterHandler<int32> SourceRegister(Context);
-			const uint16 DestRegisterIdx = Context.DecodeU16();
-
-			int32* DestReg = Context.GetOutputRegister<int32>(DataSetIndex, DestRegisterIdx) + Context.ValidInstanceIndexStart;
-
-			FMemory::StreamingMemcpy(DestReg, SourceRegister.GetDest(), sizeof(int32) * Context.ValidInstanceCount);
+			break;
 		}
-	};
+		ExecFunction(Context);
+	}
+}
 
-	// functor for splatting a constant value to an output register
-	template<typename InputHandler>
-	struct FSplatOp
+void ConditionalAddOutputWrapper(FVectorVMCodeOptimizerContext& Context, bool& OuterAdded)
+{
+	if (!OuterAdded)
 	{
-		void VM_FORCEINLINE operator()(FVectorVMContext& Context, uint16 DataSetIndex)
-		{
-			InputHandler SourceRegister(Context);
-			const uint16 DestRegisterIdx = Context.DecodeU16();
-
-			int32* DestReg = Context.GetOutputRegister<int32>(DataSetIndex, DestRegisterIdx) + Context.ValidInstanceIndexStart;
-
-			const int32 SourceValue = SourceRegister.Get();
-			const int32 InstanceVectorCount = FMath::DivideAndRoundDown(Context.ValidInstanceCount, VECTOR_WIDTH_FLOATS);
+		Context.Write<FVectorVMExecFunction>(ExecBatchedOutput);
+		OuterAdded = true;
+	}
+}
 
-			if (InstanceVectorCount)
-			{
-				const VectorRegisterInt SplatValue = MakeVectorRegisterInt(SourceValue, SourceValue, SourceValue, SourceValue);
+EVectorVMOp BatchedOutputOptimization(EVectorVMOp Op, FVectorVMCodeOptimizerContext& Context)
+{
+	if (!GbBatchVMOutput)
+	{
+		return Op;
+	}
 
-				for (int32 VectorIt = 0; VectorIt < InstanceVectorCount; ++VectorIt)
-				{
-					VectorIntStore(SplatValue, DestReg + VectorIt * VECTOR_WIDTH_FLOATS);
-				}
-			}
+	bool OuterAdded = false;
+	bool HasValidOp = true;
 
-			for (int32 InstanceIt = InstanceVectorCount * VECTOR_WIDTH_FLOATS; InstanceIt < Context.ValidInstanceCount; ++InstanceIt)
-			{
-				DestReg[InstanceIt] = SourceValue;
-			}
-		}
-	};
-
-	// performs the operation of copying data from a temporary register to an output register under the assumption
-	// that the validity of each instance is uniform (valid or not).
-	template<typename PopulateOp>
-	static VM_FORCEINLINE void DoRegisterKernelFixedValid(FVectorVMContext& Context)
+	while (HasValidOp)
 	{
-		const uint16 DataSetIndex = Context.DecodeU16();
-		Context.DecodeU16(); // DestIndexRegisterIdx
-		const uint16 AccumulatedOpCount = Context.DecodeU16();
-
-		// if none of the instances are valid, then don't bother writing anything
-		if (!Context.ValidInstanceCount)
+		switch (Op)
 		{
-			// todo we should early out of this case rather than keep parsing the code
-			for (uint16 OpIt = 0; OpIt < AccumulatedOpCount; ++OpIt)
-			{
-				FRegisterHandler<int32> Dummy(Context);
-				Context.DecodeU16(); // DestRegisterIdx
-			}
+		case EVectorVMOp::outputdata_half:
+			ConditionalAddOutputWrapper(Context, OuterAdded);
+			FScalarKernelWriteOutputIndexed<float, FFloat16, 2>::Optimize(Context);
+			break;
 
-			return;
-		}
+		case EVectorVMOp::outputdata_int32:
+			ConditionalAddOutputWrapper(Context, OuterAdded);
+			FScalarKernelWriteOutputIndexed<int32, int32, 1>::Optimize(Context);
+			break;
 
-		// for each of our ops, copy the data from the working register to the output
-		const int32 DataSize = sizeof(int32) * Context.ValidInstanceCount;
+		case EVectorVMOp::outputdata_float:
+			ConditionalAddOutputWrapper(Context, OuterAdded);
+			FScalarKernelWriteOutputIndexed<float, float, 0>::Optimize(Context);
+			break;
 
-		for (uint16 OpIt = 0; OpIt < AccumulatedOpCount; ++OpIt)
-		{
-			PopulateOp()(Context, DataSetIndex);
+		default:
+			HasValidOp = false;
+			break;
 		}
-	}
 
-	// performs the operation of copying data from a temporary register to an output register without foreknowledge
-	// of the validity of individual instances
-	static VM_FORCEINLINE void DoRegisterKernelVariableValid(FVectorVMContext& Context)
-	{
-		// if we found that all of the instances are valid, then just run the fixed version
-		if (Context.ValidInstanceUniform)
+		if (HasValidOp)
 		{
-			DoRegisterKernelFixedValid<FCopyOp>(Context);
-			return;
+			Op = Context.BaseContext.DecodeOp();
 		}
+	}
 
-		const uint16 DataSetIndex = Context.DecodeU16();
-		const uint16 DestIndexRegisterIdx = Context.DecodeU16();
-		const uint16 AccumulatedOpCount = Context.DecodeU16();
-
-		FDataSetMeta& DataSetMeta = Context.GetDataSetMeta(DataSetIndex);
-		
-		const int8* DestIndexReg = reinterpret_cast<const int8*>(Context.GetTempRegister(DestIndexRegisterIdx));
-
-		//
-		// VectorIntStore(		- unaligned writes of 16 bytes to our Destination; note that this maneuver requires us to have
-		//						our output buffers padded out to 16 bytes!
-		//	VectorIntShuffle(	- swizzle our source register to pack the valid entries at the beginning, with 0s at the end
-		//    Source,			- source data
-		//    ShuffleMask),		- result of the VectorMaskBits done in the acquireindex, int8/VectorRegister of input
-		//  Destination);
-		for (uint16 OpIt = 0; OpIt < AccumulatedOpCount; ++OpIt)
-		{
-			const RegisterType* Source = FRegisterHandler<RegisterType>(Context).GetDest();
-			int32* DestReg = Context.GetOutputRegister<int32>(DataSetIndex, Context.DecodeU16()) + Context.ValidInstanceIndexStart;
-
-			// the number of instances that we're expecting to write.  it is important that we keep track of it because when we
-			// get down to the end we need to switch from the shuffled approach to a scalar approach so that we don't
-			// overwrite the indexed output that another parallel context might have written to
-			int32 WritesRemaining = Context.ValidInstanceCount;
-			int32 SourceIt = 0;
-
-			// vector shuffle path writes 4 at a time (though the trailing elements may not be valid) until we have to move over
-			// to the scalar version for fear of overwriting our neighbors
-			while (WritesRemaining >= VECTOR_WIDTH_FLOATS)
-			{
-				check(SourceIt * VECTOR_WIDTH_FLOATS < Context.NumInstances);
-
-				const int8 ShuffleMask = DestIndexReg[SourceIt];
-				const int8 AdvanceCount = FMath::CountBits(ShuffleMask);
-
-				VectorIntStore(VectorIntShuffle(Source[SourceIt], VectorVMConstants::RegisterShuffleMask[ShuffleMask]), DestReg);
-
-				DestReg += AdvanceCount;
-				WritesRemaining -= AdvanceCount;
-
-				++SourceIt;
-			}
-
-			// scalar path that will read 4 values and write one at a time to the output based on the valid mask
-			while (WritesRemaining)
-			{
-				const int8 ShuffleMask = DestIndexReg[SourceIt];
-				const int8 AdvanceCount = FMath::CountBits(ShuffleMask);
-				if (AdvanceCount)
-				{
-					int32 RawSourceData[VECTOR_WIDTH_FLOATS];
-
-					VectorIntStore(Source[SourceIt], RawSourceData);
-
-					for (int32 ScalarIt = 0; ScalarIt < 4; ++ScalarIt)
-					{
-						if (!!(ShuffleMask & (1 << ScalarIt)))
-						{
-							*DestReg = RawSourceData[ScalarIt];
-							++DestReg;
-						}
-					}
-
-					WritesRemaining -= AdvanceCount;
-				}
-				++SourceIt;
-			}
-		}
+	if (OuterAdded)
+	{
+		Context.Write<FVectorVMExecFunction>(nullptr);
 	}
+	return Op;
+}
 
+// Optimization managed by GbBatchPackVMOutput via PackedOutputOptimization()
+// Looks for the common pattern of an acquireindex op followed by a number of associated outputdata ops.  The
+// stock operation is to write an index into a temporary register, and then have the different outputs streams
+// write into the indexed location.  This optimization does a number of things:
+// -first we check if 'validity' is uniform or not, if it is we can have a fast path of both figuring out how many
+// indices we need, as well as how to write the output (if we find that they are all invalid, then we don't need to do anything!)
+// -if we need to evaluate the validity of each element we quickly count up the number (with vector intrinsics) and
+// grab a block of the indices (rather than one at a time)
+// -rather than storing the indices to use, we store a int8 mask which indicates a valid flag for each of the next 4 samples
+// -outputs are then written to depending on their source and their frequency:
+//		-uniform sources will be splatted to all valid entries
+//		-variable sources will be packed into the available slots
+struct FBatchedWriteIndexedOutput
+{
 	// acquires a batch of indices from the provided CounterHandler.  If we're running in parallel, then we'll need to use
 	// atomics to guarantee our place in the list of indices.
 	template<bool bParallel>
@@ -2710,13 +2759,11 @@ struct FBatchedWriteIndexedOutput
 		for (int32 LoopIt = 0; LoopIt < LoopCount; ++LoopIt)
 		{
 			// input register needs to be padded to allow for 16 byte reads; but mask out the ones beyond NumInstances
-			const VectorRegister Mask = VectorVMConstants::RemainderMask[FMath::Min(VECTOR_WIDTH_FLOATS, Remainder)];
-
-			const int8 ValidMask = static_cast<int8>(VectorMaskBits(VectorSelect(Mask, ValidReader.GetAndAdvance(), GlobalVectorConstants::FloatZero)));
+			uint8 ValidMask = static_cast<uint8>(VectorMaskBits(ValidReader.GetAndAdvance()));
+			ValidMask &= ~(0xFF << FMath::Min(VECTOR_WIDTH_FLOATS, Remainder));
 			ValidCount += FMath::CountBits(ValidMask);
-
+			
 			DestAddr[LoopIt] = ValidMask;
-
 			Remainder -= VECTOR_WIDTH_FLOATS;
 		}
 
@@ -2787,6 +2834,337 @@ struct FBatchedWriteIndexedOutput
 		}
 	}
 
+	bool IsOutput(EVectorVMOp Op) const
+	{
+		return (Op == EVectorVMOp::outputdata_int32
+			|| Op == EVectorVMOp::outputdata_float
+			|| Op == EVectorVMOp::outputdata_half);
+	}
+
+	bool IsValidEnd(EVectorVMOp Op) const
+	{
+		return (Op == EVectorVMOp::done
+			|| Op == EVectorVMOp::acquireindex);
+	}
+
+	static bool SkipIfEmpty(FVectorVMContext& Context)
+	{
+		if (!Context.ValidInstanceCount)
+		{
+			Context.SkipCode(2 * sizeof(uint16)); // don't need DataSetIndex or DestIndexRegisterIdx
+			const uint16 AccumulatedOpCount = Context.DecodeU16();
+
+			Context.SkipCode(AccumulatedOpCount * 2 * sizeof(uint16));
+			return true;
+		}
+
+		return false;
+	}
+
+	VM_FORCEINLINE static void SplatElement(int32 OutputCount, int32 ElementValue, int32* RESTRICT OutputElements)
+	{
+		const int32 OutputCountWide = AlignDown(OutputCount, VECTOR_WIDTH_FLOATS);
+
+		if (OutputCountWide)
+		{
+			const VectorRegisterInt SplatValue = MakeVectorRegisterInt(ElementValue, ElementValue, ElementValue, ElementValue);
+
+			for (int32 i = 0; i < OutputCountWide; i += VECTOR_WIDTH_FLOATS)
+			{
+				VectorIntStore(SplatValue, OutputElements + i);
+			}
+		}
+
+		for (int32 i = OutputCountWide; i < OutputCount; ++i)
+		{
+			OutputElements[i] = ElementValue;
+		}
+	}
+
+	VM_FORCEINLINE static void SplatElement(int32 OutputCount, float ElementValue, float* RESTRICT OutputElements)
+	{
+		const int32 OutputCountWide = AlignDown(OutputCount, VECTOR_WIDTH_FLOATS);
+
+		if (OutputCountWide)
+		{
+			const VectorRegister SplatValue = MakeVectorRegister(ElementValue, ElementValue, ElementValue, ElementValue);
+
+			for (int32 i = 0; i < OutputCountWide; i += VECTOR_WIDTH_FLOATS)
+			{
+				VectorStore(SplatValue, OutputElements + i);
+			}
+		}
+
+		for (int32 i = OutputCountWide; i < OutputCount; ++i)
+		{
+			OutputElements[i] = ElementValue;
+		}
+	}
+
+
+	VM_FORCEINLINE static void SplatElement(int32 OutputCount, float ElementValue, FFloat16* RESTRICT OutputElements)
+	{
+		uint16 TargetValue;
+		FPlatformMath::StoreHalf(&TargetValue, ElementValue);
+		uint32 TargetValue32 = uint32(TargetValue) | ((uint32(TargetValue) << 16));
+		const VectorRegister SplatValue = MakeVectorRegister(TargetValue32, TargetValue32, TargetValue32, TargetValue32);
+
+		while (OutputCount > 7)
+		{
+			VectorStore(SplatValue, OutputElements);
+			OutputElements += 8;
+			OutputCount -= 8;
+		}
+
+		while (OutputCount > 0)
+		{
+			*((uint16*)OutputElements) = TargetValue;
+			OutputElements++;
+			OutputCount--;
+		}
+	}
+
+	template<typename T>
+	VM_FORCEINLINE static void CopyElements(int32 ElementCount, const T* RESTRICT SourceElements, T* RESTRICT OutputElements)
+	{
+		memcpy(OutputElements, SourceElements, ElementCount * sizeof(T));
+	}
+
+	VM_FORCEINLINE static void CopyElements(int32 ElementCount, const float* RESTRICT SourceElements, FFloat16* RESTRICT OutputElements)
+	{
+		const float* Src = SourceElements;
+		uint16* Dst = (uint16*)OutputElements;
+		while (ElementCount > 7)
+		{
+			FPlatformMath::WideVectorStoreHalf(Dst, Src);
+			Dst += 8;
+			Src += 8;
+			ElementCount -= 8;
+		}
+		while (ElementCount > 3)
+		{
+			FPlatformMath::VectorStoreHalf(Dst, Src);
+			Dst += 4;
+			Src += 4;
+			ElementCount -= 4;
+		}
+		while (ElementCount > 0)
+		{
+			FPlatformMath::StoreHalf(Dst, *Src);
+			Dst += 1;
+			Src += 1;
+			ElementCount -= 1;
+		}
+	}
+
+	template<typename SourceType, typename TargetType>
+	VM_FORCEINLINE static void ScalarShuffleElements(int32 ElementCount, const SourceType* RESTRICT SourceElements, const int8* RESTRICT ValidMask, TargetType* RESTRICT OutputElements)
+	{
+		constexpr int32 ElementsPerMask = 4;
+
+		// scalar path that will read 4 values and write one at a time to the output based on the valid mask
+		int32 MaskIt = 0;
+
+		while (ElementCount)
+		{
+			const int8 ShuffleMask = ValidMask[MaskIt];
+			const int8 AdvanceCount = FMath::CountBits(ShuffleMask);
+			check(AdvanceCount >= 0 && AdvanceCount <= 4);
+			if (AdvanceCount)
+			{
+				for (int32 ScalarIt = 0; ScalarIt < ElementsPerMask; ++ScalarIt)
+				{
+					if (!!(ShuffleMask & (1 << ScalarIt)))
+					{
+						*OutputElements = SourceElements[MaskIt * ElementsPerMask + ScalarIt];
+						++OutputElements;
+					}
+				}
+
+				ElementCount -= AdvanceCount;
+			}
+			++MaskIt;
+		}
+	}
+
+	VM_FORCEINLINE static void ScalarShuffleElements(int32 ElementCount, const float* RESTRICT SourceElements, const int8* RESTRICT ValidMask, FFloat16* RESTRICT OutputElements)
+	{
+		constexpr int32 ElementsPerMask = 4;
+
+		// scalar path that will read 4 values and write one at a time to the output based on the valid mask
+		int32 MaskIt = 0;
+
+		while (ElementCount)
+		{
+			checkSlow(ElementCount > 0);
+			const int8 ShuffleMask = ValidMask[MaskIt];
+			const int8 AdvanceCount = FMath::CountBits(ShuffleMask);
+			if (AdvanceCount)
+			{
+				for (int32 ScalarIt = 0; ScalarIt < ElementsPerMask; ++ScalarIt)
+				{
+					if (!!(ShuffleMask & (1 << ScalarIt)))
+					{
+						FPlatformMath::StoreHalf((uint16*)OutputElements, SourceElements[MaskIt * ElementsPerMask + ScalarIt]);
+						++OutputElements;
+					}
+				}
+
+				ElementCount -= AdvanceCount;
+			}
+			++MaskIt;
+		}
+	}
+
+	VM_FORCEINLINE static void VectorShuffleElements(int32 ElementCount, const int32* RESTRICT SourceElements, const int8* RESTRICT ValidMask, int32* RESTRICT OutputElements)
+	{
+		int32 SourceIt = 0;
+
+		const VectorRegisterInt* RESTRICT SourceVectors = reinterpret_cast<const VectorRegisterInt* RESTRICT>(SourceElements);
+
+		// vector shuffle path writes 4 at a time (though the trailing elements may not be valid) until we have to move over
+		// to the scalar version for fear of overwriting our neighbors
+		while (ElementCount >= VECTOR_WIDTH_FLOATS)
+		{
+			const int8 ShuffleMask = ValidMask[SourceIt];
+			const int8 AdvanceCount = FMath::CountBits(ShuffleMask);
+			check(AdvanceCount >= 0 && AdvanceCount <= 4);
+
+			//
+			// VectorIntStore(		- unaligned writes of 16 bytes to our Destination; note that this maneuver requires us to have
+			//						our output buffers padded out to 16 bytes!
+			//	VectorIntShuffle(	- swizzle our source register to pack the valid entries at the beginning, with 0s at the end
+			//    Source,			- source data
+			//    ShuffleMask),		- result of the VectorMaskBits done in the acquireindex, int8/VectorRegister of input
+			//  Destination);
+			VectorIntStore(VectorIntShuffle(SourceVectors[SourceIt], VectorVMConstants::RegisterShuffleMask[ShuffleMask]), OutputElements);
+
+			OutputElements += AdvanceCount;
+			ElementCount -= AdvanceCount;
+
+			++SourceIt;
+		}
+
+		ScalarShuffleElements(ElementCount, SourceElements + VECTOR_WIDTH_FLOATS * SourceIt, ValidMask + SourceIt, OutputElements);
+	}
+
+	VM_FORCEINLINE static void VectorShuffleElements(int32 ElementCount, const float* RESTRICT SourceElements, const int8* RESTRICT ValidMask, FFloat16* RESTRICT OutputElements)
+	{
+		int32 SourceIt = 0;
+
+		const VectorRegisterInt* RESTRICT SourceVectors = reinterpret_cast<const VectorRegisterInt* RESTRICT>(SourceElements);
+
+		// vector shuffle path writes 4 at a time (though the trailing elements may not be valid) until we have to move over
+		// to the scalar version for fear of overwriting our neighbors
+		while (ElementCount >= VECTOR_WIDTH_FLOATS)
+		{
+			const int8 ShuffleMask = ValidMask[SourceIt];
+			const int8 AdvanceCount = FMath::CountBits(ShuffleMask);
+			check(AdvanceCount >= 0 && AdvanceCount <= 4);
+
+			// VectorIntStore(		- unaligned writes of 16 bytes to our Destination; note that this maneuver requires us to have
+			//						our output buffers padded out to 16 bytes!
+			//	VectorIntShuffle(	- swizzle our source register to pack the valid entries at the beginning, with 0s at the end
+			//    Source,			- source data
+			//    ShuffleMask),		- result of the VectorMaskBits done in the acquireindex, int8/VectorRegister of input
+			//  Destination);
+
+			const VectorRegisterInt ShuffledFloats = VectorIntShuffle(SourceVectors[SourceIt], VectorVMConstants::RegisterShuffleMask[ShuffleMask]);
+			FPlatformMath::VectorStoreHalf((uint16*)OutputElements, (float*)&ShuffledFloats);
+
+			OutputElements += AdvanceCount;
+			ElementCount -= AdvanceCount;
+
+			++SourceIt;
+		}
+
+		ScalarShuffleElements(ElementCount, SourceElements + VECTOR_WIDTH_FLOATS * SourceIt, ValidMask + SourceIt, OutputElements);
+	}
+
+	template<typename SourceType, typename TargetType, int32 TypeOffset>
+	static void CopyConstantToOutput(FVectorVMContext& Context)
+	{
+		if (SkipIfEmpty(Context))
+		{
+			return;
+		}
+
+		const uint16 DataSetIndex = Context.DecodeU16();
+		Context.SkipCode(sizeof(uint16)); // don't need DestIndexRegisterIdx
+		const uint16 AccumulatedOpCount = Context.DecodeU16();
+
+		FDataSetMeta& DataSetMeta = Context.GetDataSetMeta(DataSetIndex);
+
+		for (uint16 OpIt = 0; OpIt < AccumulatedOpCount; ++OpIt)
+		{
+			FConstantHandler<SourceType> SourceHandler(Context);
+			const uint16 DestRegisterIndex = Context.DecodeU16();
+
+			TargetType* RESTRICT OutputElements = Context.GetOutputRegister<TargetType, TypeOffset>(DataSetIndex, DestRegisterIndex) + Context.ValidInstanceIndexStart;
+
+			SplatElement(Context.ValidInstanceCount, SourceHandler.Constant, OutputElements);
+		}
+	}
+
+	template<typename SourceType, typename TargetType, int32 TypeOffset>
+	static void CopyRegisterToOutput(FVectorVMContext& Context)
+	{
+		if (SkipIfEmpty(Context))
+		{
+			return;
+		}
+
+		const uint16 DataSetIndex = Context.DecodeU16();
+		Context.SkipCode(sizeof(uint16)); // don't need DestIndexRegisterIdx
+		const uint16 AccumulatedOpCount = Context.DecodeU16();
+
+		check(Context.ValidInstanceCount == Context.NumInstances);
+		FDataSetMeta& DataSetMeta = Context.GetDataSetMeta(DataSetIndex);
+
+		for (uint16 OpIt = 0; OpIt < AccumulatedOpCount; ++OpIt)
+		{
+			FRegisterHandler<SourceType> SourceHandler(Context);
+			const uint16 DestRegisterIndex = Context.DecodeU16();
+
+			TargetType* RESTRICT OutputElements = Context.GetOutputRegister<TargetType, TypeOffset>(DataSetIndex, DestRegisterIndex) + Context.ValidInstanceIndexStart;
+
+			CopyElements(Context.ValidInstanceCount, SourceHandler.GetDest(), OutputElements);
+		}
+	}
+
+	template<typename SourceType, typename TargetType, int32 TypeOffset>
+	static void ShuffleRegisterToOutput(FVectorVMContext& Context)
+	{
+		if (SkipIfEmpty(Context))
+		{
+			return;
+		}
+		else if (Context.ValidInstanceUniform)
+		{
+			CopyRegisterToOutput<SourceType, TargetType, TypeOffset>(Context);
+			return;
+		}
+
+		const uint16 DataSetIndex = Context.DecodeU16();
+		const uint16 DestIndexRegisterIdx = Context.DecodeU16();
+		const uint16 AccumulatedOpCount = Context.DecodeU16();
+
+		FDataSetMeta& DataSetMeta = Context.GetDataSetMeta(DataSetIndex);
+
+		const int8* RESTRICT DestIndexReg = reinterpret_cast<const int8*>(Context.GetTempRegister(DestIndexRegisterIdx));
+
+		for (uint16 OpIt = 0; OpIt < AccumulatedOpCount; ++OpIt)
+		{
+			FRegisterHandler<SourceType> SourceRegister(Context);
+			TargetType* RESTRICT DestReg = Context.GetOutputRegister<TargetType, TypeOffset>(DataSetIndex, Context.DecodeU16()) + Context.ValidInstanceIndexStart;
+
+			// the number of instances that we're expecting to write.  it is important that we keep track of it because when we
+			// get down to the end we need to switch from the shuffled approach to a scalar approach so that we don't
+			// overwrite the indexed output that another parallel context might have written to
+			VectorShuffleElements(Context.ValidInstanceCount, SourceRegister.GetDest(), DestIndexReg, DestReg);
+		}
+	}
+
 	bool OptimizeBatch(FVectorVMCodeOptimizerContext& Context)
 	{
 		const int32 BatchedOpCount = BatchedOps.Num();
@@ -2801,20 +3179,41 @@ struct FBatchedWriteIndexedOutput
 			if (!AccumulatedOpCount)
 				continue;
 
-			switch (BatchEntry.Key.SrcOpType)
+			// matrix of options based on the:
+			//	Op (are we outputting ints, floats or halfs)
+			//	SrcOpType (are we copying a constant or a register over)
+			// additionally, if we know that the index is constant, then we know at worst we're doing a copy
+
+			if (BatchEntry.Key.SrcOpType == SRCOP_RRC)
 			{
-			case SRCOP_RRR:
-				if (AcquireIndexConstant)
+				switch (BatchEntry.Key.Op)
 				{
-					Context.Write<FVectorVMExecFunction>(DoRegisterKernelFixedValid<FCopyOp>);
+				case EVectorVMOp::outputdata_float: Context.Write<FVectorVMExecFunction>(CopyConstantToOutput<float, float, 0>); break;
+				case EVectorVMOp::outputdata_int32: Context.Write<FVectorVMExecFunction>(CopyConstantToOutput<int32, int32, 1>); break;
+				case EVectorVMOp::outputdata_half: Context.Write<FVectorVMExecFunction>(CopyConstantToOutput<float, FFloat16, 2>); break;
+				default: check(0);
 				}
-				else
+			}
+			else if (AcquireIndexConstant)
+			{
+				switch (BatchEntry.Key.Op)
 				{
-					Context.Write<FVectorVMExecFunction>(DoRegisterKernelVariableValid);
+				case EVectorVMOp::outputdata_float: Context.Write<FVectorVMExecFunction>(CopyRegisterToOutput<float, float, 0>); break;
+				case EVectorVMOp::outputdata_int32: Context.Write<FVectorVMExecFunction>(CopyRegisterToOutput<int32, int32, 1>); break;
+				case EVectorVMOp::outputdata_half: Context.Write<FVectorVMExecFunction>(CopyRegisterToOutput<float, FFloat16, 2>); break;
+				default: check(0);
+				}
+			}
+			else
+			{
+				check(BatchEntry.Key.SrcOpType == SRCOP_RRR);
+				switch (BatchEntry.Key.Op)
+				{
+				case EVectorVMOp::outputdata_float: Context.Write<FVectorVMExecFunction>(ShuffleRegisterToOutput<int32, int32, 0>); break;
+				case EVectorVMOp::outputdata_int32: Context.Write<FVectorVMExecFunction>(ShuffleRegisterToOutput<int32, int32, 1>); break;
+				case EVectorVMOp::outputdata_half: Context.Write<FVectorVMExecFunction>(ShuffleRegisterToOutput<float, FFloat16, 2>); break;
+				default: check(0);
 				}
-				break;
-			case SRCOP_RRC:	Context.Write<FVectorVMExecFunction>(DoRegisterKernelFixedValid<FSplatOp<FConstantHandler<int32>>>); break;
-			default: check(0); break;
 			}
 
 			Context.Write(BatchEntry.Key.DataSetIndex);
@@ -2830,9 +3229,10 @@ struct FBatchedWriteIndexedOutput
 		return true;
 	}
 
-	bool ExtractOp(FVectorVMCodeOptimizerContext& Context)
+	bool ExtractOp(EVectorVMOp Op, FVectorVMCodeOptimizerContext& Context)
 	{
 		FOpKey Key;
+		Key.Op = Op;
 		Key.SrcOpType = Context.BaseContext.DecodeSrcOperandTypes();
 		Key.DataSetIndex = Context.DecodeU16();
 		Key.DestIndexRegisterIdx = Context.DecodeU16();
@@ -2867,6 +3267,7 @@ private:
 		uint16 DestIndexRegisterIdx;
 		uint16 DataSetIndex;
 		uint8 SrcOpType;
+		EVectorVMOp Op;
 	};
 
 	struct FOpValue
@@ -2881,14 +3282,13 @@ private:
 		{
 			return A.DestIndexRegisterIdx == B.DestIndexRegisterIdx
 				&& A.DataSetIndex == B.DataSetIndex
-				&& A.SrcOpType == B.SrcOpType;
+				&& A.SrcOpType == B.SrcOpType
+				&& A.Op == B.Op;
 		}
 
 		static VM_FORCEINLINE uint32 GetKeyHash(const FOpKey& Key)
 		{
-			return HashCombine(
-				HashCombine(GetTypeHash(Key.DestIndexRegisterIdx), GetTypeHash(Key.DataSetIndex)),
-				GetTypeHash(Key.SrcOpType));
+			return HashCombine(Key.DestIndexRegisterIdx | (Key.DataSetIndex << 16), Key.SrcOpType | (static_cast<uint8>(Key.Op) << 8));
 		}
 	};
 
@@ -2915,9 +3315,9 @@ EVectorVMOp PackedOutputOptimization(EVectorVMOp Op, FVectorVMCodeOptimizerConte
 
 		Op = Context.BaseContext.DecodeOp();
 
-		while (BatchValid && Op == EVectorVMOp::outputdata_32bit)
+		while (BatchValid && BatchedOutputOp.IsOutput(Op))
 		{
-			BatchValid = BatchedOutputOp.ExtractOp(Context);
+			BatchValid = BatchedOutputOp.ExtractOp(Op, Context);
 			Op = Context.BaseContext.DecodeOp();
 		}
 
@@ -2932,6 +3332,75 @@ EVectorVMOp PackedOutputOptimization(EVectorVMOp Op, FVectorVMCodeOptimizerConte
 	return Op;
 }
 
+void ExecBatchedInput(FVectorVMContext& Context)
+{
+	while (true)
+	{
+		FVectorVMExecFunction ExecFunction = reinterpret_cast<FVectorVMExecFunction>(Context.DecodePtr());
+		if (ExecFunction == nullptr)
+		{
+			break;
+		}
+		ExecFunction(Context);
+	}
+}
+
+void ConditionalAddInputWrapper(FVectorVMCodeOptimizerContext& Context, bool& OuterAdded)
+{
+	if (!OuterAdded)
+	{
+		Context.Write<FVectorVMExecFunction>(ExecBatchedInput);
+		OuterAdded = true;
+	}
+}
+
+EVectorVMOp BatchedInputOptimization(EVectorVMOp Op, FVectorVMCodeOptimizerContext& Context)
+{
+	if (!GbBatchVMInput)
+	{
+		return Op;
+	}
+
+	bool OuterAdded = false;
+	bool HasValidOp = true;
+
+	while (HasValidOp)
+	{
+		switch (Op)
+		{
+		case EVectorVMOp::inputdata_half:
+			ConditionalAddInputWrapper(Context, OuterAdded);
+			FVectorKernelReadInput<FFloat16, 2>::Optimize(Context);
+			break;
+
+		case EVectorVMOp::inputdata_int32:
+			ConditionalAddInputWrapper(Context, OuterAdded);
+			FVectorKernelReadInput<int32, 1>::Optimize(Context);
+			break;
+
+		case EVectorVMOp::inputdata_float:
+			ConditionalAddInputWrapper(Context, OuterAdded);
+			FVectorKernelReadInput<float, 0>::Optimize(Context);
+			break;
+
+		default:
+			HasValidOp = false;
+			break;
+		}
+
+		if (HasValidOp)
+		{
+			Op = Context.BaseContext.DecodeOp();
+		}
+	}
+
+	if (OuterAdded)
+	{
+		Context.Write<FVectorVMExecFunction>(nullptr);
+	}
+	return Op;
+}
+
 EVectorVMOp SafeMathOptimization(EVectorVMOp Op, FVectorVMCodeOptimizerContext& Context)
 {
 	if (!GbSafeOptimizedKernels)
@@ -2958,10 +3427,10 @@ void VectorVM::OptimizeByteCode(const uint8* ByteCode, TArray<uint8>& OptimizedC
 {
 	OptimizedCode.Empty();
 
-//-TODO: Support unaligned writes & little endian
+	//-TODO:7 Support unaligned writes & little endian
 #if PLATFORM_SUPPORTS_UNALIGNED_LOADS && PLATFORM_LITTLE_ENDIAN
 
-	if ( !GbOptimizeVMByteCode || (ByteCode == nullptr) )
+	if (!GbOptimizeVMByteCode || (ByteCode == nullptr))
 	{
 		return;
 	}
@@ -2971,6 +3440,8 @@ void VectorVM::OptimizeByteCode(const uint8* ByteCode, TArray<uint8>& OptimizedC
 	// add any optimization filters in here, useful so what we can isolate optimizations with CVars
 	FVectorVMCodeOptimizerContext::OptimizeVMFunction VMFilters[] =
 	{
+		//BatchedInputOptimization,
+		//BatchedOutputOptimization,
 		PackedOutputOptimization,
 		SafeMathOptimization,
 	};
@@ -3066,9 +3537,15 @@ void VectorVM::OptimizeByteCode(const uint8* ByteCode, TArray<uint8>& OptimizedC
 			case EVectorVMOp::i2b: FVectorKernelIntToBool::Optimize(Context); break;
 			case EVectorVMOp::b2i: FVectorKernelBoolToInt::Optimize(Context); break;
 
-			case EVectorVMOp::outputdata_32bit:	FScalarKernelWriteOutputIndexed<int32>::Optimize(Context);	break;
-			case EVectorVMOp::inputdata_32bit: FVectorKernelReadInput<int32>::Optimize(Context); break;
-			case EVectorVMOp::inputdata_noadvance_32bit: FVectorKernelReadInputNoAdvance<int32>::Optimize(Context); break;
+			case EVectorVMOp::outputdata_half:	FScalarKernelWriteOutputIndexed<float, FFloat16, 2>::Optimize(Context);	break;
+			case EVectorVMOp::inputdata_half: FVectorKernelReadInput<FFloat16, 2>::Optimize(Context); break;
+			case EVectorVMOp::outputdata_int32:	FScalarKernelWriteOutputIndexed<int32, int32, 1>::Optimize(Context);	break;
+			case EVectorVMOp::inputdata_int32: FVectorKernelReadInput<int32, 1>::Optimize(Context); break;
+			case EVectorVMOp::outputdata_float:	FScalarKernelWriteOutputIndexed<float, float, 0>::Optimize(Context);	break;
+			case EVectorVMOp::inputdata_float: FVectorKernelReadInput<float, 0>::Optimize(Context); break;
+			case EVectorVMOp::inputdata_noadvance_int32: FVectorKernelReadInputNoAdvance<int32, 1>::Optimize(Context); break;
+			case EVectorVMOp::inputdata_noadvance_float: FVectorKernelReadInputNoAdvance<float, 0>::Optimize(Context); break;
+			case EVectorVMOp::inputdata_noadvance_half: FVectorKernelReadInputNoAdvance<FFloat16, 2>::Optimize(Context); break;
 			case EVectorVMOp::acquireindex: FScalarKernelAcquireCounterIndex::Optimize(Context); break;
 			case EVectorVMOp::external_func_call: FKernelExternalFunctionCall::Optimize(Context); break;
 
diff --git a/Engine/Source/Runtime/VectorVM/Public/VectorVM.h b/Engine/Source/Runtime/VectorVM/Public/VectorVM.h
--- a/Engine/Source/Runtime/VectorVM/Public/VectorVM.h
+++ b/Engine/Source/Runtime/VectorVM/Public/VectorVM.h
@@ -125,9 +125,15 @@ enum class EVectorVMOp : uint8
 	b2i,
 
 	// data read/write
-	inputdata_32bit,
-	inputdata_noadvance_32bit,
-	outputdata_32bit,
+	inputdata_float,
+	inputdata_int32,
+	inputdata_half,
+	inputdata_noadvance_float,
+	inputdata_noadvance_int32,
+	inputdata_noadvance_half,
+	outputdata_float,
+	outputdata_int32,
+	outputdata_half,
 	acquireindex,
 
 	external_func_call,
@@ -156,8 +162,11 @@ enum class EVectorVMOp : uint8
 //Some require RWBuffer like support.
 struct FDataSetMeta
 {
-	uint8*RESTRICT*RESTRICT InputRegisters;
-	uint8*RESTRICT*RESTRICT OutputRegisters;
+	TArrayView<uint8 const* RESTRICT const> InputRegisters;
+	TArrayView<uint8 const* RESTRICT const> OutputRegisters;
+
+	uint32 InputRegisterTypeOffsets[3];
+	uint32 OutputRegisterTypeOffsets[3];
 
 	int32 DataSetAccessIndex;	// index for individual elements of this set
 
@@ -185,9 +194,7 @@ struct FDataSetMeta
 	FORCEINLINE void UnlockFreeTable();
 
 	FDataSetMeta()
-		: InputRegisters(nullptr)
-		, OutputRegisters(nullptr)
-		, DataSetAccessIndex(INDEX_NONE)
+		: DataSetAccessIndex(INDEX_NONE)
 		, InstanceOffset(INDEX_NONE)
 		, IDTable(nullptr)
 		, FreeIDTable(nullptr)
@@ -200,8 +207,8 @@ struct FDataSetMeta
 
 	FORCEINLINE void Reset()
 	{
-		InputRegisters = nullptr;
-		OutputRegisters = nullptr;
+		InputRegisters = TArrayView<uint8 const* RESTRICT const>();
+		OutputRegisters = TArrayView<uint8 const* RESTRICT const>();
 		DataSetAccessIndex = INDEX_NONE;
 		InstanceOffset = INDEX_NONE;
 		IDTable = nullptr;
@@ -212,7 +219,7 @@ struct FDataSetMeta
 		IDAcquireTag = INDEX_NONE;
 	}
 
-	FORCEINLINE void Init(uint8*RESTRICT *RESTRICT InInputRegisters, uint8*RESTRICT *RESTRICT InOutputRegisters, int32 InInstanceOffset, TArray<int32>* InIDTable, TArray<int32>* InFreeIDTable, int32* InNumFreeIDs, int32* InMaxUsedID, int32 InIDAcquireTag, TArray<int32>* InSpawnedIDsTable)
+	FORCEINLINE void Init(const TArrayView<uint8 const* RESTRICT const>& InInputRegisters, const TArrayView<uint8 const* RESTRICT const>& InOutputRegisters, int32 InInstanceOffset, TArray<int32>* InIDTable, TArray<int32>* InFreeIDTable, int32* InNumFreeIDs, int32* InMaxUsedID, int32 InIDAcquireTag, TArray<int32>* InSpawnedIDsTable)
 	{
 		InputRegisters = InInputRegisters;
 		OutputRegisters = InOutputRegisters;
@@ -352,17 +359,19 @@ public:
 
 	FORCEINLINE FDataSetMeta& GetDataSetMeta(int32 DataSetIndex) { return DataSetMetaTable[DataSetIndex]; }
 	FORCEINLINE uint8 * RESTRICT GetTempRegister(int32 RegisterIndex) { return TempRegTable.GetData() + TempRegisterSize * RegisterIndex; }
-	template<typename T>
+	template<typename T, int TypeOffset>
 	FORCEINLINE T* RESTRICT GetInputRegister(int32 DataSetIndex, int32 RegisterIndex) 
 	{
 		FDataSetMeta& Meta = GetDataSetMeta(DataSetIndex);
-		return ((T*)Meta.InputRegisters[RegisterIndex]) + Meta.InstanceOffset;
+		uint32 Offset = Meta.InputRegisterTypeOffsets[TypeOffset];
+		return ((T*)Meta.InputRegisters[Offset + RegisterIndex]) + Meta.InstanceOffset;
 	}
-	template<typename T>
+	template<typename T, int TypeOffset>
 	FORCEINLINE T* RESTRICT GetOutputRegister(int32 DataSetIndex, int32 RegisterIndex) 
 	{ 
 		FDataSetMeta& Meta = GetDataSetMeta(DataSetIndex);
-		return  ((T*)Meta.OutputRegisters[RegisterIndex]) + Meta.InstanceOffset;
+		uint32 Offset = Meta.OutputRegisterTypeOffsets[TypeOffset];
+		return  ((T*)Meta.OutputRegisters[Offset + RegisterIndex]) + Meta.InstanceOffset;
 	}
 
 	int32 GetNumInstances() const { return NumInstances; }
@@ -382,6 +391,7 @@ public:
 	FORCEINLINE uint64 DecodeU64() { uint64 v = Code[7]; v = v << 8 | Code[6]; v = v << 8 | Code[5]; v = v << 8 | Code[4]; v = v << 8 | Code[3]; v = v << 8 | Code[2]; v = v << 8 | Code[1]; v = v << 8 | Code[0]; Code += 8; return INTEL_ORDER64(v); }
 #endif
 	FORCEINLINE uintptr_t DecodePtr() { return (sizeof(uintptr_t) == 4) ? DecodeU32() : DecodeU64(); }
+	FORCEINLINE void SkipCode(int64 Count) { Code += Count; }
 
 	/** Decode the next operation contained in the bytecode. */
 	FORCEINLINE EVectorVMOp DecodeOp()
