Add local jump table to VectorVM's optimized byte code to reduce memory footprint (~%30 lower)

diff --git a/Engine/Source/Runtime/VectorVM/Private/VectorVM.cpp b/Engine/Source/Runtime/VectorVM/Private/VectorVM.cpp
--- a/Engine/Source/Runtime/VectorVM/Private/VectorVM.cpp
+++ b/Engine/Source/Runtime/VectorVM/Private/VectorVM.cpp
@@ -239,6 +239,9 @@ struct FVectorVMCodeOptimizerContext
 	{
 		BaseContext.PrepareForExec(0, 0, nullptr, nullptr, nullptr, nullptr, TArrayView<FDataSetMeta>(), 0, false);
 		BaseContext.PrepareForChunk(ByteCode, 0, 0);
+
+		// write out a jump table offset that we'll fill in when the table is encoded (EncodeJumpTable())
+		Write<uint32>(0);
 	}
 	FVectorVMCodeOptimizerContext(const FVectorVMCodeOptimizerContext&) = delete;
 	FVectorVMCodeOptimizerContext(const FVectorVMCodeOptimizerContext&&) = delete;
@@ -258,6 +261,14 @@ struct FVectorVMCodeOptimizerContext
 		reinterpret_cast<T&>(OptimizedCode[OptimizedCode.AddUninitialized(sizeof(T))]) = v;
 	}
 
+	template<>
+	void Write<FVectorVMExecFunction>(const FVectorVMExecFunction& Function)
+	{
+		const int32 JumpTableIndex = JumpTable.AddUnique(Function);
+		check(JumpTableIndex <= TNumericLimits<uint8>::Max());
+		Write<uint8>(JumpTableIndex);
+	}
+
 	struct FOptimizerCodeState
 	{
 		uint8 const* BaseContextCode;
@@ -279,8 +290,44 @@ struct FVectorVMCodeOptimizerContext
 		OptimizedCode.SetNum(State.OptimizedCodeLength, false /* allowShrink */);
 	}
 
+	// Jump table is encoded at the end of the optimized code, with the first int32 in the byte code
+	// stream being the start offset
+
+	static const FVectorVMExecFunction* DecodeJumpTable(const uint8*& OptimizedByteCode)
+	{
+		const uint32 JumpTableOffset = *reinterpret_cast<const uint32*>(OptimizedByteCode);
+		const FVectorVMExecFunction* JumpTable = reinterpret_cast<const FVectorVMExecFunction*>(OptimizedByteCode + JumpTableOffset);
+
+		OptimizedByteCode += sizeof(uint32);
+
+		return JumpTable;
+	}
+
+	void EncodeJumpTable()
+	{
+		const int32 JumpTableCount = JumpTable.Num();
+		check(JumpTableCount > 0);
+		check(JumpTableCount <= TNumericLimits<uint8>::Max());
+
+		if (JumpTableCount <= 0 || JumpTableCount > TNumericLimits<uint8>::Max())
+		{
+			// if the jump table is too big, then clear out the OptimizedCode and we'll just have to run the unoptimized path
+			OptimizedCode.Reset();
+			return;
+		}
+
+		// write the offset to the jump table into the reserved slot
+		const uint32 JumpTableOffset = OptimizedCode.Num();
+		*reinterpret_cast<uint32*>(OptimizedCode.GetData()) = JumpTableOffset;
+
+		const int32 JumpTableSize = JumpTableCount * sizeof(FVectorVMExecFunction);
+		OptimizedCode.AddUninitialized(JumpTableSize);
+		FMemory::Memcpy(OptimizedCode.GetData() + JumpTableOffset, JumpTable.GetData(), JumpTableSize);
+	}
+
 	FVectorVMContext&		BaseContext;
 	TArray<uint8>&			OptimizedCode;
+	TArray<FVectorVMExecFunction, TInlineAllocator<256>> JumpTable;
 	const TArrayView<uint8>	ExternalFunctionRegisterCounts;
 	const int32				StartInstance = 0;
 };
@@ -2427,6 +2474,8 @@ void VectorVM::Exec(
 	const bool bParallel = NumBatches > 1;
 	const bool bUseOptimizedByteCode = (OptimizedByteCode != nullptr) && GbUseOptimizedVMByteCode;
 
+	const FVectorVMExecFunction* OptimizedJumpTable = bUseOptimizedByteCode ? FVectorVMCodeOptimizerContext::DecodeJumpTable(OptimizedByteCode) : nullptr;
+
 	auto ExecChunkBatch = [&](int32 BatchIdx)
 	{
 		//SCOPE_CYCLE_COUNTER(STAT_VVMExecChunk);
@@ -2455,7 +2504,7 @@ void VectorVM::Exec(
 
 				while (true)
 				{
-					FVectorVMExecFunction ExecFunction = reinterpret_cast<FVectorVMExecFunction>(Context.DecodePtr());
+					FVectorVMExecFunction ExecFunction = OptimizedJumpTable[Context.DecodeU8()];
 					if (ExecFunction == nullptr)
 					{
 						break;
@@ -3576,6 +3625,8 @@ void VectorVM::OptimizeByteCode(const uint8* ByteCode, TArray<uint8>& OptimizedC
 		}
 	} while (Op != EVectorVMOp::done);
 	Context.Write<FVectorVMExecFunction>(nullptr);
+
+	Context.EncodeJumpTable();
 #endif //PLATFORM_SUPPORTS_UNALIGNED_LOADS && PLATFORM_LITTLE_ENDIAN
 }
 
