Fix for member function specialization on android with VectorVM

diff --git a/Engine/Source/Runtime/VectorVM/Private/VectorVM.cpp b/Engine/Source/Runtime/VectorVM/Private/VectorVM.cpp
--- a/Engine/Source/Runtime/VectorVM/Private/VectorVM.cpp
+++ b/Engine/Source/Runtime/VectorVM/Private/VectorVM.cpp
@@ -261,8 +261,7 @@ struct FVectorVMCodeOptimizerContext
 		reinterpret_cast<T&>(OptimizedCode[OptimizedCode.AddUninitialized(sizeof(T))]) = v;
 	}
 
-	template<>
-	void Write<FVectorVMExecFunction>(const FVectorVMExecFunction& Function)
+	void WriteExecFunction(const FVectorVMExecFunction& Function)
 	{
 		const int32 JumpTableIndex = JumpTable.AddUnique(Function);
 		check(JumpTableIndex <= TNumericLimits<uint8>::Max());
@@ -579,7 +578,7 @@ struct TUnaryKernelHandler
 {
 	static void Optimize(FVectorVMCodeOptimizerContext& Context)
 	{
-		Context.Write<FVectorVMExecFunction>(Exec);
+		Context.WriteExecFunction(Exec);
 		Arg0Handler::Optimize(Context);
 		DstHandler::Optimize(Context);
 	}
@@ -602,7 +601,7 @@ struct TBinaryKernelHandler
 {
 	static void Optimize(FVectorVMCodeOptimizerContext& Context)
 	{
-		Context.Write<FVectorVMExecFunction>(Exec);
+		Context.WriteExecFunction(Exec);
 		Arg0Handler::Optimize(Context);
 		Arg1Handler::Optimize(Context);
 		DstHandler::Optimize(Context);
@@ -628,7 +627,7 @@ struct TTrinaryKernelHandler
 {
 	static void Optimize(FVectorVMCodeOptimizerContext& Context)
 	{
-		Context.Write<FVectorVMExecFunction>(Exec);
+		Context.WriteExecFunction(Exec);
 		Arg0Handler::Optimize(Context);
 		Arg1Handler::Optimize(Context);
 		Arg2Handler::Optimize(Context);
@@ -1121,7 +1120,7 @@ struct FVectorKernelExecutionIndex
 {
 	static void Optimize(FVectorVMCodeOptimizerContext& Context)
 	{
-		Context.Write<FVectorVMExecFunction>(Exec);
+		Context.WriteExecFunction(Exec);
 		FRegisterHandler<VectorRegisterInt>::Optimize(Context);
 	}
 
@@ -1146,7 +1145,7 @@ struct FVectorKernelEnterStatScope
 	static void Optimize(FVectorVMCodeOptimizerContext& Context)
 	{
 #if STATS
-		Context.Write<FVectorVMExecFunction>(Exec);
+		Context.WriteExecFunction(Exec);
 		FConstantHandler<int32>::Optimize(Context);
 #else
 		// just skip the op if we don't have stats enabled
@@ -1172,7 +1171,7 @@ struct FVectorKernelExitStatScope
 	static void Optimize(FVectorVMCodeOptimizerContext& Context)
 	{
 #if STATS
-		Context.Write<FVectorVMExecFunction>(Exec);
+		Context.WriteExecFunction(Exec);
 #endif
 	}
 		
@@ -1463,7 +1462,7 @@ struct FScalarKernelAcquireID
 {
 	static void Optimize(FVectorVMCodeOptimizerContext& Context)
 	{
-		Context.Write<FVectorVMExecFunction>(Exec);
+		Context.WriteExecFunction(Exec);
 		Context.Write(Context.DecodeU16());		// DataSetIndex
 		Context.Write(Context.DecodeU16());		// IDIndexReg
 		Context.Write(Context.DecodeU16());		// IDTagReg
@@ -1520,7 +1519,7 @@ struct FScalarKernelUpdateID
 {
 	static void Optimize(FVectorVMCodeOptimizerContext& Context)
 	{
-		Context.Write<FVectorVMExecFunction>(Exec);
+		Context.WriteExecFunction(Exec);
 		Context.Write(Context.DecodeU16());		// DataSetIndex
 		Context.Write(Context.DecodeU16());		// InstanceIDRegisterIndex
 		Context.Write(Context.DecodeU16());		// InstanceIndexRegisterIndex
@@ -1579,7 +1578,7 @@ struct FVectorKernelReadInput
 {
 	static void Optimize(FVectorVMCodeOptimizerContext& Context)
 	{
-		Context.Write<FVectorVMExecFunction>(Exec);
+		Context.WriteExecFunction(Exec);
 		Context.Write(Context.DecodeU16());	// DataSetIndex
 		Context.Write(Context.DecodeU16());	// InputRegisterIdx
 		Context.Write(Context.DecodeU16());	// DestRegisterIdx
@@ -1613,7 +1612,7 @@ struct FVectorKernelReadInput<FFloat16, 2>
 {
 	static void Optimize(FVectorVMCodeOptimizerContext& Context)
 	{
-		Context.Write<FVectorVMExecFunction>(Exec);
+		Context.WriteExecFunction(Exec);
 		Context.Write(Context.DecodeU16());	// DataSetIndex
 		Context.Write(Context.DecodeU16());	// InputRegisterIdx
 		Context.Write(Context.DecodeU16());	// DestRegisterIdx
@@ -1658,7 +1657,7 @@ struct FVectorKernelReadInputNoAdvance
 {
 	static void Optimize(FVectorVMCodeOptimizerContext& Context)
 	{
-		Context.Write<FVectorVMExecFunction>(Exec);
+		Context.WriteExecFunction(Exec);
 		Context.Write(Context.DecodeU16());	// DataSetIndex
 		Context.Write(Context.DecodeU16());	// InputRegisterIdx
 		Context.Write(Context.DecodeU16());	// DestRegisterIdx
@@ -1690,7 +1689,7 @@ struct FVectorKernelReadInputNoAdvance<FFloat16, 2>
 {
 	static void Optimize(FVectorVMCodeOptimizerContext& Context)
 	{
-		Context.Write<FVectorVMExecFunction>(Exec);
+		Context.WriteExecFunction(Exec);
 		Context.Write(Context.DecodeU16());	// DataSetIndex
 		Context.Write(Context.DecodeU16());	// InputRegisterIdx
 		Context.Write(Context.DecodeU16());	// DestRegisterIdx
@@ -1747,8 +1746,8 @@ struct FScalarKernelWriteOutputIndexed
 		const uint32 SrcOpTypes = Context.BaseContext.DecodeSrcOperandTypes();
 		switch (SrcOpTypes)
 		{
-			case SRCOP_RRR: Context.Write<FVectorVMExecFunction>(DoKernel<FRegisterHandler<SourceType>>); break;
-			case SRCOP_RRC:	Context.Write<FVectorVMExecFunction>(DoKernel<FConstantHandler<SourceType>>); break;
+			case SRCOP_RRR: Context.WriteExecFunction(DoKernel<FRegisterHandler<SourceType>>); break;
+			case SRCOP_RRC:	Context.WriteExecFunction(DoKernel<FConstantHandler<SourceType>>); break;
 			default: check(0); break;
 		};
 
@@ -1806,8 +1805,8 @@ struct FScalarKernelWriteOutputIndexed<float, FFloat16, 2>
 		const uint32 SrcOpTypes = Context.BaseContext.DecodeSrcOperandTypes();
 		switch (SrcOpTypes)
 		{
-		case SRCOP_RRR: Context.Write<FVectorVMExecFunction>(DoKernel<FRegisterHandler<float>>); break;
-		case SRCOP_RRC:	Context.Write<FVectorVMExecFunction>(DoKernel<FConstantHandler<float>>); break;
+		case SRCOP_RRR: Context.WriteExecFunction(DoKernel<FRegisterHandler<float>>); break;
+		case SRCOP_RRC:	Context.WriteExecFunction(DoKernel<FConstantHandler<float>>); break;
 		default: check(0); break;
 		};
 
@@ -1932,8 +1931,8 @@ struct FScalarKernelAcquireCounterIndex
 		const uint32 SrcOpType = Context.BaseContext.DecodeSrcOperandTypes();
 		switch (SrcOpType)
 		{
-			case SRCOP_RRR: Context.Write<FVectorVMExecFunction>(FScalarKernelAcquireCounterIndex::ExecOptimized<SRCOP_RRR>); break;
-			case SRCOP_RRC: Context.Write<FVectorVMExecFunction>(FScalarKernelAcquireCounterIndex::ExecOptimized<SRCOP_RRC>); break;
+			case SRCOP_RRR: Context.WriteExecFunction(FScalarKernelAcquireCounterIndex::ExecOptimized<SRCOP_RRR>); break;
+			case SRCOP_RRC: Context.WriteExecFunction(FScalarKernelAcquireCounterIndex::ExecOptimized<SRCOP_RRC>); break;
 			default: check(0); break;
 		}
 
@@ -1985,7 +1984,7 @@ struct FKernelExternalFunctionCall
 	{
 		const uint32 ExternalFuncIdx = Context.DecodeU8();
 
-		Context.Write<FVectorVMExecFunction>(Exec);
+		Context.WriteExecFunction(Exec);
 		Context.Write<uint8>(ExternalFuncIdx);
 
 		const int32 NumRegisters = Context.ExternalFunctionRegisterCounts[ExternalFuncIdx];
@@ -2700,7 +2699,7 @@ void ConditionalAddOutputWrapper(FVectorVMCodeOptimizerContext& Context, bool& O
 {
 	if (!OuterAdded)
 	{
-		Context.Write<FVectorVMExecFunction>(ExecBatchedOutput);
+		Context.WriteExecFunction(ExecBatchedOutput);
 		OuterAdded = true;
 	}
 }
@@ -2747,7 +2746,7 @@ EVectorVMOp BatchedOutputOptimization(EVectorVMOp Op, FVectorVMCodeOptimizerCont
 
 	if (OuterAdded)
 	{
-		Context.Write<FVectorVMExecFunction>(nullptr);
+		Context.WriteExecFunction(nullptr);
 	}
 	return Op;
 }
@@ -2870,8 +2869,8 @@ struct FBatchedWriteIndexedOutput
 
 		switch (SrcOpType)
 		{
-		case SRCOP_RRR: Context.Write<FVectorVMExecFunction>(IndexExecOptimized<SRCOP_RRR>); break;
-		case SRCOP_RRC: Context.Write<FVectorVMExecFunction>(IndexExecOptimized<SRCOP_RRC>); break;
+		case SRCOP_RRR: Context.WriteExecFunction(IndexExecOptimized<SRCOP_RRR>); break;
+		case SRCOP_RRC: Context.WriteExecFunction(IndexExecOptimized<SRCOP_RRC>); break;
 		default: check(0); break;
 		}
 
@@ -3243,9 +3242,9 @@ struct FBatchedWriteIndexedOutput
 			{
 				switch (BatchEntry.Key.Op)
 				{
-				case EVectorVMOp::outputdata_float: Context.Write<FVectorVMExecFunction>(CopyConstantToOutput<float, float, 0>); break;
-				case EVectorVMOp::outputdata_int32: Context.Write<FVectorVMExecFunction>(CopyConstantToOutput<int32, int32, 1>); break;
-				case EVectorVMOp::outputdata_half: Context.Write<FVectorVMExecFunction>(CopyConstantToOutput<float, FFloat16, 2>); break;
+				case EVectorVMOp::outputdata_float: Context.WriteExecFunction(CopyConstantToOutput<float, float, 0>); break;
+				case EVectorVMOp::outputdata_int32: Context.WriteExecFunction(CopyConstantToOutput<int32, int32, 1>); break;
+				case EVectorVMOp::outputdata_half: Context.WriteExecFunction(CopyConstantToOutput<float, FFloat16, 2>); break;
 				default: check(0);
 				}
 			}
@@ -3253,9 +3252,9 @@ struct FBatchedWriteIndexedOutput
 			{
 				switch (BatchEntry.Key.Op)
 				{
-				case EVectorVMOp::outputdata_float: Context.Write<FVectorVMExecFunction>(CopyRegisterToOutput<float, float, 0>); break;
-				case EVectorVMOp::outputdata_int32: Context.Write<FVectorVMExecFunction>(CopyRegisterToOutput<int32, int32, 1>); break;
-				case EVectorVMOp::outputdata_half: Context.Write<FVectorVMExecFunction>(CopyRegisterToOutput<float, FFloat16, 2>); break;
+				case EVectorVMOp::outputdata_float: Context.WriteExecFunction(CopyRegisterToOutput<float, float, 0>); break;
+				case EVectorVMOp::outputdata_int32: Context.WriteExecFunction(CopyRegisterToOutput<int32, int32, 1>); break;
+				case EVectorVMOp::outputdata_half: Context.WriteExecFunction(CopyRegisterToOutput<float, FFloat16, 2>); break;
 				default: check(0);
 				}
 			}
@@ -3264,9 +3263,9 @@ struct FBatchedWriteIndexedOutput
 				check(BatchEntry.Key.SrcOpType == SRCOP_RRR);
 				switch (BatchEntry.Key.Op)
 				{
-				case EVectorVMOp::outputdata_float: Context.Write<FVectorVMExecFunction>(ShuffleRegisterToOutput<int32, int32, 0>); break;
-				case EVectorVMOp::outputdata_int32: Context.Write<FVectorVMExecFunction>(ShuffleRegisterToOutput<int32, int32, 1>); break;
-				case EVectorVMOp::outputdata_half: Context.Write<FVectorVMExecFunction>(ShuffleRegisterToOutput<float, FFloat16, 2>); break;
+				case EVectorVMOp::outputdata_float: Context.WriteExecFunction(ShuffleRegisterToOutput<int32, int32, 0>); break;
+				case EVectorVMOp::outputdata_int32: Context.WriteExecFunction(ShuffleRegisterToOutput<int32, int32, 1>); break;
+				case EVectorVMOp::outputdata_half: Context.WriteExecFunction(ShuffleRegisterToOutput<float, FFloat16, 2>); break;
 				default: check(0);
 				}
 			}
@@ -3404,7 +3403,7 @@ void ConditionalAddInputWrapper(FVectorVMCodeOptimizerContext& Context, bool& Ou
 {
 	if (!OuterAdded)
 	{
-		Context.Write<FVectorVMExecFunction>(ExecBatchedInput);
+		Context.WriteExecFunction(ExecBatchedInput);
 		OuterAdded = true;
 	}
 }
@@ -3451,7 +3450,7 @@ EVectorVMOp BatchedInputOptimization(EVectorVMOp Op, FVectorVMCodeOptimizerConte
 
 	if (OuterAdded)
 	{
-		Context.Write<FVectorVMExecFunction>(nullptr);
+		Context.WriteExecFunction(nullptr);
 	}
 	return Op;
 }
@@ -3624,7 +3623,7 @@ void VectorVM::OptimizeByteCode(const uint8* ByteCode, TArray<uint8>& OptimizedC
 				return;//BAIL
 		}
 	} while (Op != EVectorVMOp::done);
-	Context.Write<FVectorVMExecFunction>(nullptr);
+	Context.WriteExecFunction(nullptr);
 
 	Context.EncodeJumpTable();
 #endif //PLATFORM_SUPPORTS_UNALIGNED_LOADS && PLATFORM_LITTLE_ENDIAN
