Remove unused variables in VectorVM that causes compiler error with clang 10 (Merged from CL 14359176).

#rnx

diff --git a/Engine/Source/Runtime/VectorVM/Private/VectorVM.cpp b/Engine/Source/Runtime/VectorVM/Private/VectorVM.cpp
--- a/Engine/Source/Runtime/VectorVM/Private/VectorVM.cpp
+++ b/Engine/Source/Runtime/VectorVM/Private/VectorVM.cpp
@@ -1223,7 +1223,6 @@ struct FVectorKernelRandom : public TUnaryVectorKernel<FVectorKernelRandom>
 {
 	static void VM_FORCEINLINE DoKernel(FVectorVMContext& Context, VectorRegister* RESTRICT Dst, VectorRegister Src0)
 	{
-		const float rm = RAND_MAX;
 		//EEK!. Improve this. Implement GPU style seeded rand instead of this.
 		VectorRegister Result = MakeVectorRegister(Context.RandStream.GetFraction(),
 			Context.RandStream.GetFraction(),
@@ -1238,7 +1237,6 @@ struct FVectorKernelRandomGauss : public TBinaryVectorKernel<FVectorKernelRandom
 {
 	static void VM_FORCEINLINE DoKernel(FVectorVMContext& Context, VectorRegister* RESTRICT Dst, VectorRegister Src0, VectorRegister Src1)
 	{
-		const float rm = RAND_MAX;
 		VectorRegister Result = MakeVectorRegister(Context.RandStream.GetFraction(),
 			Context.RandStream.GetFraction(),
 			Context.RandStream.GetFraction(),
@@ -2153,7 +2151,6 @@ struct FScalarIntKernelRandom : public TUnaryScalarIntKernel<FScalarIntKernelRan
 {
 	static void VM_FORCEINLINE DoKernel(FVectorVMContext& Context, int32* RESTRICT Dst, int32 Src0)
 	{
-		const float rm = RAND_MAX;
 		//EEK!. Improve this. Implement GPU style seeded rand instead of this.
 		*Dst = static_cast<int32>(Context.RandStream.GetFraction() * Src0);
 	}
