Standardize step() implementation between CPU and GPU scripts making the existing step(x) implement step(0, x)

diff --git a/Engine/Plugins/FX/Niagara/Source/NiagaraCore/Private/NiagaraCustomVersion.cpp b/Engine/Plugins/FX/Niagara/Source/NiagaraCore/Private/NiagaraCustomVersion.cpp
--- a/Engine/Plugins/FX/Niagara/Source/NiagaraCore/Private/NiagaraCustomVersion.cpp
+++ b/Engine/Plugins/FX/Niagara/Source/NiagaraCore/Private/NiagaraCustomVersion.cpp
@@ -8,4 +8,4 @@ const FGuid FNiagaraCustomVersion::GUID(0xFCF57AFA, 0x50764283, 0xB9A9E658, 0xFF
 // Register the custom version with core
 FCustomVersionRegistration GRegisterNiagaraCustomVersion(FNiagaraCustomVersion::GUID, FNiagaraCustomVersion::LatestVersion, TEXT("NiagaraVer"));
 
-const FGuid FNiagaraCustomVersion::LatestScriptCompileVersion(0x38F94A17, 0xC44A414C, 0x9F8AFB2C, 0x605115EB);
+const FGuid FNiagaraCustomVersion::LatestScriptCompileVersion(0x0EEBCD8A, 0xE45341FF, 0x84B53A98, 0x46EC4EC1);
diff --git a/Engine/Plugins/FX/Niagara/Source/NiagaraEditor/Private/NiagaraEditorCommon.cpp b/Engine/Plugins/FX/Niagara/Source/NiagaraEditor/Private/NiagaraEditorCommon.cpp
--- a/Engine/Plugins/FX/Niagara/Source/NiagaraEditor/Private/NiagaraEditorCommon.cpp
+++ b/Engine/Plugins/FX/Niagara/Source/NiagaraEditor/Private/NiagaraEditorCommon.cpp
@@ -741,9 +741,9 @@ void FNiagaraOpInfo::Init()
 		Op = &OpInfos[Idx];
 		Op->Category = CategoryText;
 		Op->FriendlyName = NSLOCTEXT("NiagaraOpInfo", "Step Name", "Step");
-		Op->Description = NSLOCTEXT("NiagaraOpInfo", "Step Desc", "Result = step(A)");
+		Op->Description = NSLOCTEXT("NiagaraOpInfo", "Step Desc", "Result = step(0, A)");
 		Op->Inputs.Add(FNiagaraOpInOutInfo(A, Type, AText, AText, DefaultStr_One));
-		Op->Outputs.Add(FNiagaraOpInOutInfo(Result, Type, ResultText, ResultText, DefaultStr_One, TEXT("step({0})")));
+		Op->Outputs.Add(FNiagaraOpInOutInfo(Result, Type, ResultText, ResultText, DefaultStr_One, TEXT("step(0, {0})")));
 		Op->BuildName(TEXT("Step"), CategoryName);
 		OpInfoMap.Add(Op->Name) = Idx;
 
diff --git a/Engine/Source/Runtime/VectorVM/Private/VectorVM.cpp b/Engine/Source/Runtime/VectorVM/Private/VectorVM.cpp
--- a/Engine/Source/Runtime/VectorVM/Private/VectorVM.cpp
+++ b/Engine/Source/Runtime/VectorVM/Private/VectorVM.cpp
@@ -968,6 +968,14 @@ struct FVectorKernelLog2 : public TUnaryVectorKernel<FVectorKernelLog2>
 	}
 };
 
+struct FVectorKernelStep : public TBinaryVectorKernel<FVectorKernelStep>
+{
+	static void VM_FORCEINLINE DoKernel(FVectorVMContext& Context, VectorRegister* RESTRICT Dst, VectorRegister Src0, VectorRegister Src1)
+	{
+		*Dst = VectorStep(VectorSubtract(Src1, Src0));
+	}
+};
+
 struct FVectorKernelClamp : public TTrinaryVectorKernel<FVectorKernelClamp>
 {
 	static void VM_FORCEINLINE DoKernel(FVectorVMContext& Context, VectorRegister* RESTRICT Dst,VectorRegister Src0,VectorRegister Src1,VectorRegister Src2)
@@ -1323,14 +1331,6 @@ struct FVectorKernelSign : public TUnaryVectorKernel<FVectorKernelSign>
 	}
 };
 
-struct FVectorKernelStep : public TUnaryVectorKernel<FVectorKernelStep>
-{
-	static void VM_FORCEINLINE DoKernel(FVectorVMContext& Context, VectorRegister* RESTRICT Dst, VectorRegister Src0)
-	{
-		*Dst = VectorStep(Src0);
-	}
-};
-
 namespace VectorVMNoise
 {
 	int32 P[512] =
