Moved VM functions over to use their accurate vectorized version (still uses approximations but uses newton-raphson iteraions to improve the estimate at some cost).  Brings individual op costs of rcp(), rsqrt() and, by association, sqrt() up by a factor of ~4 but nearly doubles the precision.  Also updated a few of the core modules (FindQuatBetween/DirectionAndLength/DirectionAndLengthSafe) to better take advantage of the fast/precise intrinsics

diff --git a/Engine/Source/Runtime/VectorVM/Private/VectorVM.cpp b/Engine/Source/Runtime/VectorVM/Private/VectorVM.cpp
--- a/Engine/Source/Runtime/VectorVM/Private/VectorVM.cpp
+++ b/Engine/Source/Runtime/VectorVM/Private/VectorVM.cpp
@@ -76,6 +76,34 @@ namespace VectorVMConstants
 	};
 };
 
+#define VM_USE_ACCURATE_VECTOR_FUNCTIONS (1)
+
+namespace VectorVMAccuracy
+{
+	VM_FORCEINLINE static VectorRegister4Float Reciprocal(const VectorRegister4Float& Src)
+	{
+#if VM_USE_ACCURATE_VECTOR_FUNCTIONS
+		return VectorReciprocalAccurate(Src);
+#else
+		return VectorReciprocal(Src);
+#endif
+	}
+
+	VM_FORCEINLINE static VectorRegister4Float ReciprocalSqrt(const VectorRegister4Float& Src)
+	{
+#if VM_USE_ACCURATE_VECTOR_FUNCTIONS
+		return VectorReciprocalSqrtAccurate(Src);
+#else
+		return VectorReciprocalSqrt(Src);
+#endif
+	}
+
+	VM_FORCEINLINE static VectorRegister4Float Sqrt(const VectorRegister4Float& Src)
+	{
+		return Reciprocal(ReciprocalSqrt(Src));
+	}
+};
+
 // helper function wrapping the SSE3 shuffle operation.  Currently implemented for some platforms, the
 // rest will just use the FPU version so as to not push the requirements up to SSE3 (currently SSE2)
 #if PLATFORM_ENABLE_VECTORINTRINSICS && PLATFORM_ALWAYS_HAS_SSE4_1
@@ -860,7 +888,7 @@ struct FVectorKernelRcp : public TUnaryVectorKernel<FVectorKernelRcp>
 {
 	static void VM_FORCEINLINE DoKernel(FVectorVMContext& Context, VectorRegister4Float* RESTRICT Dst,VectorRegister4Float Src0)
 	{
-		*Dst = VectorReciprocal(Src0);
+		*Dst = VectorVMAccuracy::Reciprocal(Src0);
 	}
 };
 
@@ -870,7 +898,7 @@ struct FVectorKernelRcpSafe : public TUnaryVectorKernel<FVectorKernelRcpSafe>
 	static void VM_FORCEINLINE DoKernel(FVectorVMContext& Context, VectorRegister4Float* RESTRICT Dst, VectorRegister4Float Src0)
 	{
 		VectorRegister4Float ValidMask = VectorCompareGT(VectorAbs(Src0), GlobalVectorConstants::SmallNumber);
-		*Dst = VectorSelect(ValidMask, VectorReciprocal(Src0), GlobalVectorConstants::FloatZero);
+		*Dst = VectorSelect(ValidMask, VectorVMAccuracy::Reciprocal(Src0), GlobalVectorConstants::FloatZero);
 	}
 };
 
@@ -878,7 +906,7 @@ struct FVectorKernelRsq : public TUnaryVectorKernel<FVectorKernelRsq>
 {
 	static void VM_FORCEINLINE DoKernel(FVectorVMContext& Context, VectorRegister4Float* RESTRICT Dst,VectorRegister4Float Src0)
 	{
-		*Dst = VectorReciprocalSqrt(Src0);
+		*Dst = VectorVMAccuracy::ReciprocalSqrt(Src0);
 	}
 };
 
@@ -888,7 +916,7 @@ struct FVectorKernelRsqSafe : public TUnaryVectorKernel<FVectorKernelRsqSafe>
 	static void VM_FORCEINLINE DoKernel(FVectorVMContext& Context, VectorRegister4Float* RESTRICT Dst, VectorRegister4Float Src0)
 	{
 		VectorRegister4Float ValidMask = VectorCompareGT(Src0, GlobalVectorConstants::SmallNumber);
-		*Dst = VectorSelect(ValidMask, VectorReciprocalSqrt(Src0), GlobalVectorConstants::FloatZero);
+		*Dst = VectorSelect(ValidMask, VectorVMAccuracy::ReciprocalSqrt(Src0), GlobalVectorConstants::FloatZero);
 	}
 };
 
@@ -896,8 +924,7 @@ struct FVectorKernelSqrt : public TUnaryVectorKernel<FVectorKernelSqrt>
 {
 	static void VM_FORCEINLINE DoKernel(FVectorVMContext& Context, VectorRegister4Float* RESTRICT Dst,VectorRegister4Float Src0)
 	{
-		// TODO: Need a SIMD sqrt!
-		*Dst = VectorReciprocal(VectorReciprocalSqrt(Src0));
+		*Dst = VectorVMAccuracy::Sqrt(Src0);
 	}
 };
 
@@ -906,7 +933,7 @@ struct FVectorKernelSqrtSafe : public TUnaryVectorKernel<FVectorKernelSqrtSafe>
 	static void VM_FORCEINLINE DoKernel(FVectorVMContext& Context, VectorRegister4Float* RESTRICT Dst, VectorRegister4Float Src0)
 	{
 		VectorRegister4Float ValidMask = VectorCompareGT(Src0, GlobalVectorConstants::SmallNumber);
-		*Dst = VectorSelect(ValidMask, VectorReciprocal(VectorReciprocalSqrt(Src0)), GlobalVectorConstants::FloatZero);
+		*Dst = VectorSelect(ValidMask, VectorVMAccuracy::Sqrt(Src0), GlobalVectorConstants::FloatZero);
 	}
 };
 
