- New Vector VM with optimizer and serialization.

diff --git a/Engine/Source/Runtime/VectorVM/Private/VectorVMExperimental.inl b/Engine/Source/Runtime/VectorVM/Private/VectorVMExperimental.inl
new file mode 100644
--- /dev/null
+++ b/Engine/Source/Runtime/VectorVM/Private/VectorVMExperimental.inl
@@ -0,0 +1,2209 @@
+// Copyright Epic Games, Inc. All Rights Reserved.
+
+/*
+external functions to improve:
+UNiagaraDataInterfaceSkeletalMesh stuff (4 wide randoms)
+
+UNiagaraDataInterfaceSkeletalMesh::GetSkinnedBoneData
+SetNumCells
+GetNumCells
+SetRenderTargetSize
+FastMatrixToQuaternion
+
+BasicSkinEmitter:
+	- GetFilteredTriangle
+	- GetSkinnedTriangleDataWS
+	- GetTriUV
+
+pathological case:
+	- ComponentRendererTest_SpawnScript_0x4A6253BF_ue: Increases temp reg count
+
+
+- batch reuse w/o going back to TaskGraph
+- prefetch instruction
+
+*/
+
+/*
+The three steps to running the new VM are:
+1. call OptimizeVectorVMScript() using the original bytecode and function bindings as input.
+It will set up the FVectorVMOptimizeContext with:
+	- New Bytecode
+	- Const Remap Table
+	- External Function Table (only containing the number of IO params... the function pointers 
+	  are set in InitVectorVMState())
+	- Some intermediate data for debugging.  These are not saved by default.
+	- Number of Constant Buffers and Temporary Registers required
+
+2. Fill out FVectorVMInitData including setting the FVectorVMOptimizeContext from step 1.  Call 
+InitVectorVMState().  This will allocate the memory required for the FVectorVMState and the first
+batch.
+
+3. call ExecVectorVMState() with the FVectorVMState from step 2.
+
+The VM operates on "Instances."  Instances are organized in groups of 4, (for now at least, 
+with AVX-2 we would use groups of 8).  A group of 4 Instances is called a "Loop."  The thread 
+hierarchy in the VM has three levels: Batches -> Chunks -> Loops.  A Batch is represented as a 
+single Async TaskGraph task.  Batches contain one or more Chunks.  Batches loop over each Chunk, 
+executing all of the Bytecode Instructions one Chunk at a time.  Chunks loop over each "Loop" 
+executing the SIMD Instructions associated with the Bytecode. 
+(More on Chunks and memory usage down below)
+
+In general, (exceptions are discussed further down), data is input into the VM through either a 
+DataSet or a Constant Buffer (ConstBuff).  The previous VM would copy all DataSet inputs into 
+temporary registers (TempRegs) before operating on them.  This VM can operate directly on the 
+Inputs from the DataSets.  Outputs are written to the Outputs in the DataSet.
+
+Each VM Instruction has n Inputs and m Outputs (m is almost always 1 except external_func_call 
+and acquire_id).  Inputs can be from one of three places: DataSetInput, ConstBuff, TempReg.  
+Instructions always output to TempRegs.  TempRegs and ConstBuffs constitute the memory required 
+for each Chunk.
+
+The optimizer takes the bytecode from the original VM as input and outputs a new bytecode for the 
+new VM.  The bytecodes are similar in that the first 100-ish Instructions are the same, but they 
+are encoded differently.  There's a few new Instructions added as well.
+
+The primary optimization concept is to minimize the number of TempRegs used in the VM in order to 
+have a significantly smaller internal state size.  The original VM's Bytecode was bookended by all 
+the input and output Instructions and internally all operations worked on TempRegs.  
+This new VM has fuse_input* Instructions that combine the input Instruction with the operation 
+Instruction so the input Instructions are mostly gone (update_id and external_func_call cannot 
+currently fuse... I could add this, but I don't think it will provide much, if any, performance 
+improvement).  Outputs are also batched to execute several at once.  Outputs that have no 
+processing on them and are effectively a "copy" from the input and are handled with a new 
+Instruction: copy_to_output, (they aren't strictly "copied" because the acquireindex Instruction 
+could change which slot they get written to).
+
+Instructions are also re-ordered to facilitate minimal TempReg usage.  The acquireindex instruction 
+figures out which Instance gets written into which slot and writes these indices into a TempReg.
+It effectively determines which Instances are discarded and which are kept.  Output Instructions 
+utilize the TempReg written to by the acquireindex Instruction to write the contents of a TempReg
+or ConstBuff to a DataSetOutput.  Output Instructions are re-ordered to execute immediately 
+following the last Instruction that uses the TempReg it writes.
+
+Constant buffers are used elsewhere in UE and have a fixed, static layout it memory. They have many
+values interleaved together.  Some of these variables are required by the VM to execute a script, 
+some are not.  This leads to gaps and random memory access when reading this sparse constant table.
+The optimizer figures out exactly which constant buffers are required for the script, and how to 
+map the constant buffer table that comes from UE into the smaller set required by the VM for a 
+particular script. This map is saved in the OptimizerContext.  The constants are copied and 
+broadcasted 4-wide to the VM's internal state in the VectorVMInit() function.
+
+Most instructions in the original VM have a control byte immediately following the opcode to
+specify whether a register used is a ConstBuff or a TempReg.  Input registers into external
+functions used a different encoding: the high bit (of a 2 byte index) is set when a register is
+temporary, or not set when it's constant.  The new VM uses a universal encoding for all registers
+everywhere: 16 bit indices, high bit set = const, otherwise temp register.
+
+Work for each execution of the VM gets broken up two ways: Batches and Chunks.  A Batch is
+effectively a "thread" and represents a single Async Task in the TaskGraph.  A Batch can further be
+split up into multiple Chunks.  The only reason to split work into Chunks is to minimize the memory 
+footprint of each Batch.  A Batch will loop over each Chunk and execute the exact same instructions 
+on each Chunk.  There are two CVars to control these: GVVMChunkSizeInBytes and 
+GVVMMaxBatchesPerExec.  Chunk size is ideally the size of the L1, (VectorVMInit() will consider a 
+little bit of overhead for the bytecode and stack when executing).  This should hopefully mean that 
+all work done by the VM fits within the L1.  The number of Batches corresponds to how many 
+TaskGraph tasks are created, and are thus a function of the the available hardware threads during 
+runtime... a difficult thing to properly load balance.
+
+For example if the L1 D$ is 32kb, and the script's Bytecode is 1kb, and we assume an overhead of 
+512 bytes: We set the GVVMChunkSizeInBytes to 32768.  The FVectorVMInit() function will do: 
+32768 - 1024 - 512 = 31232 bytes per chunk.
+[@TODO: maybe we should remove the GVVMChunkSizeInBytes and just read the L1 size directly, or if 
+GVVMChunkSizeInBytes is 0 it signals to use the L1 size]
+
+The first Batch's memory is allocated directly following the VectorVMState in the VectorVMInit() 
+function.  When ExecBatch() gets called from the TaskGraph it first attempts to reuse an existing 
+batch's memory that's already finished executing.  If it can't find an existing batch that's 
+finished it will allocate new memory for this batch.  Once a batch has its memory it will setup its
+RegisterData pointer aligned to 64 bytes.  The RegisterData pointer holds all ConstBuffs and 
+TempRegs required for the execution of a single Chunk; Batches will usually loop over several 
+Chunks.  The RegisterData holds 4-wide, 32 bit variables only (16 bytes). In RegisterData the 
+ConstBuffs come first, followed by the TempRegs.
+
+When the first Batch's memory is allocated, the required ConstBuffs are broadcasted 4-wide into the
+beginning of the Batch's RegisterData.  Only the constants that are required, as determined by 
+FVectorVMOptimize(), are set there.  When the memory is allocated for all other batches the 
+ConstBuffs are memcpy'd from the first batch.
+
+The number of Instances a Chunk can operate on is a function of the number of bytes allocated to 
+the Chunk and the number of TempRegs and ConstBuffs, (as determined by the VectorVMOptimize() 
+function), and a per-chunk overhead.  For example, a script setup in the following manner:
+
+GVVMChunkSizeInBytes: 16384 bytes
+NumBytecodeBytes:       832 bytes
+FixedChunkOverhead:     500 bytes
+NumConstBuffers:         12
+NumTempRegisters:         8
+NumDataSets:              2
+MaxRegsForExtFn           5
+
+VectorVMInit() does the following computation:
+	1.    16 bytes = NumDataSets * 8      <- track #outputs written for each DataSet
+	2.    10 bytes = MaxRegsForExtFn * 2  <- registers for ext fns get pre-decoded
+	3.   526 bytes = 500 + 16 + 10        <- BatchOverheadSize
+	1. 15026 bytes = 16384 - 532 - 526    <- max size of internal VM state
+	2.   192 bytes = 12 * 16              <- NumConstBuffers * sizeof(int32) * 4
+	3. 14834 bytes = 15026 - 192          <- #bytes remaining for TempRegs
+	4. 115 Loops   = 14843 / (8 * 16)     <- 8 TempRegs * 16 bytes per loop (4 instances)
+	5. 460 Instances Per Chunk.
+
+This particular script can execute 460 instances per chunk with a GVVMChunkSizeInBytes of 16384.
+
+As described above the new VM has a universal register encoding using 16 bit indices with the high 
+bit signifying whether the register the Instruction requires is a TempReg or ConstBuff.  This 
+allows the VM to decode which registers are used by an operation very efficiently, 4 at a time 
+using SIMD.  The equations to compute the pointers to registers required for operations are as 
+follows: (in byte offsets from the beginning of the Batch's RegisterData)
+	ConstBuff: RegisterData + 16 * ConstIdx
+	TempReg  : RegisterData + 16 * NumConsts + NumLoops * TempRegIdx
+	
+In addition to computing the offsets the "increment" variable is computed when the Instruction is 
+decoded.  The increment is 0xFFFFFFFF for TempRegs and 0 for ConstBuffs.  Each operation loops over
+registers for each Instance in the Chunk (4 at a time), and the loop index is logically AND'd with 
+the increment value such that ConstBuffs always read from the same place and TempRegs read from the 
+normal loop index.
+
+4 registers are always decoded for each Instruction regardless of how many (if any) are used by the 
+Instruction.  External functions decode their instructions into a special buffer in the batch's 
+ChunkLocalData.  If they have more than four operands, the VM loops as many times as necessary to 
+decode all the registers.  This greatly simplifies the code required to decode the registers in 
+user-defined functions.  All external functions are backwards compatible with the previous VM.
+
+Memory and batches work different on the new VM compared to the old VM.  In the old VM the Exec() 
+lambda is passed a BatchIdx which determines which instances to work on.  
+The calculation was: BatchIdx * NumChunksPerBatch * NumLoopsPerChunk.  
+This means that each BatchIdx will always work on the same set of instances.  This means that the 
+memory for each batch must always be allocated and used only once.  In times of high thread 
+contention batch memory could be sitting around unused.
+
+The new VM works differently.  Each time ExecVVMBatch() is called from the TaskGraph it tries to 
+reuse previously-allocated batches that have finished executing.  If it cannot reuse one, it will
+allocate new memory and copy the ConstBuffs from the first batch.  The function
+AssignInstancesToBatch() thread-safely grabs the next bunch of instances and assigns them to this 
+batch.
+
+There are 11 new fused_input instructions:
+	fused_input1_1 //op has 1 input operand it's an input
+	fused_input2_1 //op has 2 input operands, register 0 is an input
+	fused_input2_2 //op has 2 input operands, register 1 is an input
+	fused_input2_3 //op has 2 input operands, register 0 and 1 are inputs
+	fused_input3_1 //op has 3 input operands, register 0 is an input
+	fused_input3_2 //op has 3 input operands, register 1 is an input
+	fused_input3_3 //op has 3 input operands, register 0 and 1 are inputs
+	fused_input3_4 //op has 3 input operands, register 2 is an input
+	fused_input3_5 //op has 3 input operands, register 0 and 2 are inputs
+	fused_input3_6 //op has 3 input operands, register 1 and 2 are inputs
+	fused_input3_7 //op has 3 input operands, register 0, 1 and 2 are inputs
+
+Instructions generally have 1, 2 or 3 inputs.  They are usually TempRegs or ConstBuffs.  In some 
+cases, one or more of the TempRegs can be changed to a DataSetInput.  In order to do that, the 
+VectorVMOptimize() function injects the appropriate fusedinput operation before the Instruction.  
+For example, if the add Instruction adds ConstBuff 6 to DataSetInput 9, the VectorVMOptimize() 
+instruction will emit two Instructions: fused_input2_2, and add.  The first digit in the 
+fused_input instruction is how many operands the instruction has, and the second digit is a binary
+representation of which operands are be changed to DataSetInputs... in this case 2 = 2nd operand.  
+As another example if an fmadd Instruction was in the original Bytecode that took DataSetInputs for 
+operands 0 and 2 FVectorVMOptimize() would emit a fused_input3_5 instruction before the fmadd.
+
+acquireindex logic is different from the original VM's.  The original VM wrote which slot the
+to read from, and a -1 to indicate "skip".  This required a branch for each instance being written,
+for output instruction.  If the keep/discard boolean was distributed similar to white noise there 
+would be massive mispredict penalities.
+
+The new VM's acquireindex instruction writes which slot to write into.  This allows for branch-free 
+write Output Instructions.  For example: if it was determined that Instances 1, 3 and 4 were to be 
+discarded, acquireindex would output:
+	0, 1, 1, 2, 2, 2, 3
+
+These correspond to the slots that get written to.  So the Output instructions will loop over each 
+index, and write it into the slot specified by the index. ie:
+	write Instance 0 into slot 0
+	write Instance 1 into slot 1
+	write Instance 2 into slot 1
+	write Instance 3 into slot 2
+	write Instance 4 into slot 2
+	write Instance 5 into slot 2
+	write Instance 6 into slot 3
+
+In order to facilitate this change, aquire_id and update_id also needed to be changed.  update_id 
+and acquire_id were completely re-written in order to be lock-free.  The original VM's DataSets had 
+two separate arrays: FreeIDsTable and SpawnedIDsTable.  The FreeIDs table was pre-allocated to have 
+enough room for the persistent IDs in the worst-case situation of every single instance being freed 
+on a particular execution of the VM.  The acquire_id function pulls IDs out of the FreeIDs table 
+into a TempReg and writes them to the SpawnedIDs table.  In order for elements to be put into 
+SpawnedIDs they must first be removed from FreeIDs.  Therefore it is impossible for the counts of 
+FreeIDs + SpawnedIDs to exceed the number of instances for a particlar execution of a VM -- the 
+same number that is pre-allocated to the FreeIDs.  I removed the SpawnedIDs table and simply write 
+the SpawnedIDs to the end of the FreeIDs table.  I keep a separate index: NumSpawnedIDs in the 
+DataSet.  This allows for complete lock-free manipulation of both sets of data as it's just two 
+numbers keeping track of the two.  ie:
+
+DataSet->FreeIDsTable:
+[------------0000000000000000000000000000000**********]
+             ^ NumFreeIds                   ^ FreeIDsTable.Max() - NumSpawnedIDs
+	- represents FreeIDs
+	0 represents unused spaced
+	* represents SpawnedIDs
+
+Upon observing the Bytecode of dozens of scripts I recognized that DataSetInputs are often directly
+written to DataSetOutputs.  The new VM has a new instruction called copy_to_output which takes a 
+count and a list of DataSetInputs and DataSetOutputs and uses the acquireindex index to write 
+directly between the two without requiring a TempReg. Additionally most outputs get grouped 
+together.
+
+I also added new output_batch* instructions to write more than one output at a time:
+	output_batch8
+	output_batch7
+	output_batch4
+	output_batch3
+	output_batch2
+
+7 and 3 may seem weird, but they're there to utilize the fact that the instruction decoded looks at
+4 registers at a time, so decoding the index is free.
+It is guaranteed by the optimizer that the index for output_batch8 and output_batch4 comes from a 
+TempReg, not a ConstBuff so the decoding can be optimized.
+*/
+
+enum EVVMRegFlags
+{
+	VVMRegFlag_Int      = 1,
+	VVMRegFlag_Clean    = 32,
+	VVMRegFlag_Index    = 64,
+	VVMRegFlag_Mismatch = 128,
+};
+
+static void *VVMDefaultRealloc(void *Ptr, size_t NumBytes, const char *Filename, int LineNumber)
+{
+	return FMemory::Realloc(Ptr, NumBytes);
+}
+
+static void VVMDefaultFree(void *Ptr, const char *Filename, int LineNumber)
+{
+	return FMemory::Free(Ptr);
+}
+
+#include "./VectorVMExperimental_Serialization.inl"
+
+#ifdef NIAGARA_EXP_VM
+
+#define VVM_CACHELINE_SIZE				64
+#define VVM_CHUNK_FIXED_OVERHEAD_SIZE	512
+
+#define VVM_MIN(a, b)               ((a) < (b) ? (a) : (b))
+#define VVM_MAX(a, b)               ((a) > (b) ? (a) : (b))
+#define VVM_CLAMP(v, min, max)      ((v) < (min) ? (min) : ((v) < (max) ? (v) : (max)))
+#define VVM_ALIGN(num, alignment)   (((size_t)(num) + (alignment) - 1) & ~((alignment) - 1))
+#define VVM_ALIGN_4(num)            (((size_t)(num) + 3) & ~3)
+#define VVM_ALIGN_16(num)           (((size_t)(num) + 15) & ~15)
+#define VVM_ALIGN_32(num)           (((size_t)(num) + 31) & ~31)
+#define VVM_ALIGN_64(num)           (((size_t)(num) + 63) & ~63)
+#define VVM_ALIGN_CACHELINE(num)	(((size_t)(num) + (VVM_CACHELINE_SIZE - 1)) & ~(VVM_CACHELINE_SIZE - 1))
+
+//to avoid memset/memcpy when statically initializing sse variables
+#define VVMSet_m128Const(Name, V)                static const MS_ALIGN(16) float VVMConstVec4_##Name##4[4]  GCC_ALIGN(16) = { V, V, V, V }
+#define VVMSet_m128iConst(Name, V)               static const MS_ALIGN(16) uint32 VVMConstVec4_##Name##4i[4] GCC_ALIGN(16) = { V, V, V, V }
+#define VVMSet_m128iConst4(Name, V0, V1, V2, V3) static const MS_ALIGN(16) uint32 VVMConstVec4_##Name##4i[4] GCC_ALIGN(16) = { V0, V1, V2, V3 }	/* equiv to setr */
+
+#define VVM_m128Const(Name)  (*(VectorRegister4f *)&(VVMConstVec4_##Name##4))
+#define VVM_m128iConst(Name) (*(VectorRegister4i *)&(VVMConstVec4_##Name##4i))
+
+VVMSet_m128Const(   One             , 1.f);
+VVMSet_m128Const(   OneHalf         , 0.5f);
+VVMSet_m128Const(   Epsilon         , 1.e-8f);
+VVMSet_m128Const(   HalfPi          , 3.14159265359f * 0.5f);
+VVMSet_m128Const(   FastSinA        , 7.5894663844f);
+VVMSet_m128Const(   FastSinB        , 1.6338434578f);
+VVMSet_m128iConst(  FMask           , 0xFFFFFFFF);
+VVMSet_m128iConst4( ZeroOneTwoThree , 0, 1, 2, 3);
+VVMSet_m128iConst(  RegOffsetMask   , 0x7FFF);
+VVMSet_m128Const(   RegOneOverTwoPi , 1.f / 2.f / 3.14159265359f);
+VVMSet_m128iConst(  AlmostTwoBits   , 0x3fffffff);
+
+#define VVM_vecStep(a, b)            VectorStep(VectorSubtract(a, b))
+#define VVM_vecFloatToBool(v)        VectorCompareGT(v, VectorZeroFloat())
+#define VVM_vecBoolToFloat(v)        VectorSelect(v, VectorSet1(1.f), VectorZeroFloat());
+#define VVM_vecIntToBool(v)          VectorIntCompareGT(v, VectorSetZero())
+#define VVM_vecBoolToInt(v)          VectorIntSelect(v, VectorIntSet1(1), VectorSetZero())
+#define VVM_vecSqrtFast(v)           VectorReciprocal(VectorReciprocalSqrt(v))
+#define VVM_vecACosFast(v)           VectorATan2(VVM_vecSqrtFast(VectorMultiply(VectorSubtract(VVM_m128Const(One), v), VectorAdd(VVM_m128Const(One), v))), v)
+//safe instructions -- handle divide by zero "gracefully" by returning 0
+#define VVM_safeIns_div(v0, v1)      VectorSelect(VectorCompareGT(VectorAbs(v1), VVM_m128Const(Epsilon)), VectorDivide(v0, v1)    , VectorZeroFloat())
+#define VVM_safeIns_rcp(v)           VectorSelect(VectorCompareGT(VectorAbs(v) , VVM_m128Const(Epsilon)), VectorReciprocal(v)     , VectorZeroFloat())
+#define VVM_safe_sqrt(v)             VectorSelect(VectorCompareGT(VectorAbs(v) , VVM_m128Const(Epsilon)), VVM_vecSqrtFast(v)      , VectorZeroFloat())
+#define VVM_safe_log(v)              VectorSelect(VectorCompareGT(VectorAbs(v) , VectorZeroFloat())     , VectorLog(v)		       , VectorZeroFloat())
+#define VVM_safe_pow(v0, v1)         VectorSelect(VectorCompareGT(VectorAbs(v1), VVM_m128Const(Epsilon)), VectorPow(v0, v1)       , VectorZeroFloat())
+#define VVM_safe_rsq(v)              VectorSelect(VectorCompareGT(VectorAbs(v) , VVM_m128Const(Epsilon)), VectorReciprocalSqrt(v) , VectorZeroFloat())
+
+static void VVMMemCpy(void *dst, void *src, size_t bytes)
+{
+	unsigned char *d      = (unsigned char *)dst;
+	unsigned char *s      = (unsigned char *)src;
+	unsigned char *s_end  = s + bytes;
+	ptrdiff_t ofs_to_dest = d - s;
+	if (bytes < 16)
+	{
+		if (bytes)
+		{
+			do
+			{
+				s[ofs_to_dest] = s[0];
+				++s;
+			} while (s < s_end);
+		}
+	}
+	else
+	{
+		// do one unaligned to get us aligned for the stream out below
+		VectorRegister4i i0 = VectorIntLoad(s);
+		VectorIntStore(i0, d);
+		s += 16 + 16 - ((size_t)d & 15); // S is 16 bytes ahead 
+		while (s <= s_end)
+		{
+			i0 = VectorIntLoad(s - 16);
+			VectorIntStoreAligned(i0, s - 16 + ofs_to_dest);
+			s += 16;
+		}
+		// do one unaligned to finish the copy
+		i0 = VectorIntLoad(s_end - 16);
+		VectorIntStore(i0, s_end + ofs_to_dest - 16);
+	}      
+}
+
+static void VVMMemSet32(void *dst, uint32 val, size_t num_vals)
+{
+	check(((size_t)dst & 3) == 0);							    //must be 4 byte aligned
+	if (num_vals < 4)
+	{
+		uint32 *ptr = (uint32 *)dst;
+		uint32 *end_ptr = ptr + num_vals;
+		do
+		{
+			*ptr++ = val;
+		} while (ptr < end_ptr);
+	}
+	else
+	{
+		VectorRegister4i v4 = VectorIntSet1(val);
+		check(((size_t)dst & 3) == 0);
+		uint32 *ptr = (uint32 *)dst;
+		uint32 *end_ptr = ptr + num_vals - 4;
+		while (ptr < end_ptr) {
+			VectorIntStore(v4, ptr);
+			ptr += 4;
+		}
+		VectorIntStore(v4, end_ptr);
+	}
+}
+
+#include "./VectorVMExperimental_Optimizer.inl" 
+
+static uint8 *SetupBatchStatePtrs(FVectorVMState *VVMState, FVectorVMBatchState *BatchState, void *BatchMem)
+{
+	uint8 *BatchDataPtr = (uint8 *)VVM_ALIGN_CACHELINE(BatchMem);
+	BatchState->MallocedMemPtr                              = BatchMem;
+	BatchState->RegisterData                                = (FVecReg *)BatchDataPtr;                  BatchDataPtr += VVMState->PerBatchRegisterDataBytesRequired;
+	BatchState->ChunkLocalData.StartingOutputIdxPerDataSet	= (uint32 *) BatchDataPtr;                  BatchDataPtr += VVMState->PerBatchChunkLocalDataOutputIdxBytesRequired;
+	BatchState->ChunkLocalData.NumOutputPerDataSet          = (uint32 *) BatchDataPtr;                  BatchDataPtr += VVMState->PerBatchChunkLocalNumOutputBytesRequired;
+
+	{ //deal with the external function register decoding buffer
+		size_t PtrBeforeExtFnDecodeReg = (size_t)BatchDataPtr;
+		BatchState->ChunkLocalData.ExtFnDecodedReg.RegData       = (FVecReg **)BatchDataPtr;  BatchDataPtr += sizeof(FVecReg *) * VVMState->MaxExtFnRegisters;
+		BatchState->ChunkLocalData.ExtFnDecodedReg.RegInc        = (uint32 *)BatchDataPtr;    BatchDataPtr += sizeof(uint32)    * VVMState->MaxExtFnRegisters;
+		size_t PtrAfterExtFnDecodeReg = (size_t)BatchDataPtr;
+		check(PtrAfterExtFnDecodeReg - PtrBeforeExtFnDecodeReg == VVMState->PerBatchChunkLocalNumExtFnDecodeRegisterBytesRequired);
+	}
+	
+	{ //after everything, likely outside of what the chunk will cache, setup the almost-never-used random counters and
+		BatchState->ChunkLocalData.RandCounters = (int32 *)BatchDataPtr;
+		BatchDataPtr += sizeof(int32) * VVMState->MaxInstancesPerChunk;
+	}
+
+	if ((size_t)(BatchDataPtr - (uint8 *)BatchMem) <= VVMState->NumBytesRequiredPerBatch)
+	{
+		return BatchDataPtr;
+	}
+	else
+	{
+		VVMState->Error.Flags |= VVMErr_BatchMemory;
+		VVMState->Error.LineNum = __LINE__;
+		return nullptr;
+	}
+}
+
+static bool AssignInstancesToBatch(FVectorVMState *VVMState, FVectorVMBatchState *BatchState)
+{
+	int SanityCount = 0;
+	do
+	{
+		int OldNumAssignedInstances = VVMState->NumInstancesAssignedToBatches;
+		int MaxInstancesPerBatch    = VVMState->MaxInstancesPerChunk * VVMState->MaxChunksPerBatch;
+		int NumAssignedInstances    = FPlatformAtomics::InterlockedCompareExchange(&VVMState->NumInstancesAssignedToBatches, OldNumAssignedInstances + MaxInstancesPerBatch, OldNumAssignedInstances);
+		if (NumAssignedInstances == OldNumAssignedInstances)
+		{
+			BatchState->StartInstance = OldNumAssignedInstances;
+			BatchState->NumInstances  = MaxInstancesPerBatch;
+			if (BatchState->StartInstance + BatchState->NumInstances > VVMState->TotalNumInstances) {
+				BatchState->NumInstances = VVMState->TotalNumInstances - BatchState->StartInstance;
+			}
+			if (BatchState->NumInstances <= 0) {
+				return false; //some other thread interrupted and finished the rest of the instances, we're done.
+			}
+			VVMMemSet32(BatchState->ChunkLocalData.RandCounters, 0, VVMState->MaxInstancesPerChunk); //@TODO: complete waste.  With a more sensible prng we could remove this.
+			++BatchState->UseCount;
+			return true;
+		}
+	} while (SanityCount++ < (1 << 30));
+	return false;
+}
+
+VECTORVM_API void FreeVectorVMState(FVectorVMState *VVMState)
+{
+	if (VVMState == nullptr)
+	{
+		return;
+	}
+	VectorVMFreeFn *FreeFn = VVMState->FreeFn ? VVMState->FreeFn : VVMDefaultFree;
+	for (int i = 1; i < VVMState->NumBatches; ++i)
+	{
+		FreeFn(VVMState->BatchStates[i].MallocedMemPtr, __FILE__, __LINE__);
+	}
+	FreeFn(VVMState, __FILE__, __LINE__);
+}
+
+static void SetupRandStateForBatch(FVectorVMBatchState *BatchState)
+{
+	uint64 pcg_state = FPlatformTime::Cycles64();
+	uint64 pcg_inc   = (((uint64)BatchState << 32) ^ 0XCAFEF00DD15EA5E5U) | 1;
+	pcg_state ^= (FPlatformTime::Cycles64() << 32ULL);
+	//use psuedo-pcg to setup a state for xorwow... lol!
+	for (int i = 0; i < 5; ++i)
+	{
+		uint32 values[4];
+		for (int j = 0; j < 4; ++j)
+		{
+			uint64 old_state   = pcg_state;
+			pcg_state          = old_state * 6364136223846793005ULL + pcg_inc;
+			uint32 xor_shifted = (uint32)(((old_state >> 18U) ^ old_state) >> 27U);
+			uint32 rot         = old_state >> 59U;
+			values[j]          = (xor_shifted >> rot) | (xor_shifted << ((0U - rot) & 31));
+		}
+		VectorIntStore(MakeVectorRegisterInt(values[0], values[1], values[2], values[3]), BatchState->RandState + i);
+	}
+	BatchState->RandCounters = MakeVectorRegisterInt64(pcg_inc, pcg_state);
+	BatchState->RandStream.GenerateNewSeed();
+}
+
+static VectorRegister4i VVMXorwowStep(FVectorVMBatchState *BatchState)
+{
+	VectorRegister4i t = BatchState->RandState[4];
+	VectorRegister4i s = BatchState->RandState[0];
+	BatchState->RandState[4] = BatchState->RandState[3];
+	BatchState->RandState[3] = BatchState->RandState[2];
+	BatchState->RandState[2] = BatchState->RandState[1];
+	BatchState->RandState[1] = s;
+	t = VectorIntXor(t, VectorShiftRightImmLogical(t, 2));
+	t = VectorIntXor(t, VectorShiftLeftImm(t, 1));
+	t = VectorIntXor(t, VectorIntXor(s, VectorIntXor(s, VectorShiftLeftImm(s, 4))));
+	BatchState->RandState[0] = t;
+	BatchState->RandCounters = VectorIntAdd(BatchState->RandCounters, VectorIntSet1(362437));
+	VectorRegister4i Result = VectorIntAdd(t, VectorIntLoad(&BatchState->RandCounters));
+	return Result;
+}
+
+
+VECTORVM_API FVectorVMState *InitVectorVMState(FVectorVMInitData *InitData, FVectorVMExternalFnPerInstanceData **OutPerInstanceExtData, struct FVectorVMSerializeState *SerializeState)
+{
+	if (InitData->OptimizeContext == nullptr)
+	{
+		return nullptr;
+	}
+	if (InitData->ExtFunctionTable.Num() > 0 && InitData->OptimizeContext->MaxExtFnUsed >= InitData->ExtFunctionTable.Num())
+	{
+		check(false); //somehow the funciton table changed in between optimize() and init()
+		return nullptr;
+	}
+
+	size_t NumVVMStateBytesRequired                                    = 0; //Memory "outside" of the batches for the VVM state itself
+	size_t PerBatchRegisterDataBytesRequired                           = 0;																									                    //const + temp buffer for each batch
+	const uint32 MaxExtFnRegisters                                     = InitData->OptimizeContext->MaxExtFnRegisters == 0 ? 0 : VVM_ALIGN_4(InitData->OptimizeContext->MaxExtFnRegisters + 3); //we decode 4 at a time, so if we need any, we need a multiuple of 4
+	const size_t PerBatchChunkLocalDataOutputIdxBytesRequired	       = sizeof(uint32)                                        * InitData->OptimizeContext->NumOutputDataSets;                  //chunk local bytes required for instance offset
+	const size_t PerBatchChunkLocalNumOutputBytesRequired		       = sizeof(uint32)                                        * InitData->OptimizeContext->NumOutputDataSets;                  //chunk local bytes for num outputs
+	const size_t ConstantBufferSize					                   = sizeof(FVecReg)                                       * InitData->OptimizeContext->NumConstsRemapped;
+	const size_t PerBatchChunkLocalNumExtFnDecodeRegisterBytesRequired = (sizeof(FVecReg *) + sizeof(uint32)) * MaxExtFnRegisters; 
+
+	const size_t BatchOverheadSize                                     = ConstantBufferSize + 
+		                                                                 PerBatchChunkLocalDataOutputIdxBytesRequired + 
+		                                                                 PerBatchChunkLocalNumOutputBytesRequired + 
+		                                                                 PerBatchChunkLocalNumExtFnDecodeRegisterBytesRequired +
+		                                                                 InitData->OptimizeContext->NumBytecodeBytes +
+		                                                                 VVM_CHUNK_FIXED_OVERHEAD_SIZE;
+	static const uint32 MaxChunksPerBatch                              = 4; //*MUST BE POW 2* arbitrary 4 chunks per batch... this is harder to load balance because it depends on CPU cores available during execution
+	static_assert(MaxChunksPerBatch > 0 && (MaxChunksPerBatch & (MaxChunksPerBatch - 1)) == 0);
+	int NumBatches                                                     = 1;
+	int NumChunksPerBatch                                              = (int)MaxChunksPerBatch;
+	uint32 MaxLoopsPerChunk                                            = 0;
+	{ //compute the number of bytes required per batch
+		const uint32 TotalNumLoopsRequired              = VVM_MAX(((uint32)InitData->NumInstances + 3) >> 2, 1);
+		const size_t NumBytesRequiredPerLoop            = sizeof(FVecReg) * InitData->OptimizeContext->NumTempRegisters;
+		if (BatchOverheadSize + 64 > GVVMChunkSizeInBytes) {
+			//either the chunk size is way too small, or there's an insane number of consts or data sets required... we just revert to the previous VM's default
+			MaxLoopsPerChunk                            = 128 >> 2;
+			NumBatches                                  = (TotalNumLoopsRequired + MaxLoopsPerChunk - 1) / MaxLoopsPerChunk;
+		}
+		else
+		{
+			size_t NumBytesPerBatchAvailableForTempRegs = GVVMChunkSizeInBytes - BatchOverheadSize;
+			size_t TotalNumLoopBytesRequired            = VVM_ALIGN(TotalNumLoopsRequired, MaxChunksPerBatch) * NumBytesRequiredPerLoop;
+			if (NumBytesPerBatchAvailableForTempRegs < TotalNumLoopBytesRequired)
+			{
+				//Not everything fits into a single chunk, so we have to compute everything here
+				int NumChunksRequired                   = (int)(TotalNumLoopBytesRequired + NumBytesPerBatchAvailableForTempRegs - 1) / (int)NumBytesPerBatchAvailableForTempRegs;
+				check(NumChunksRequired > 1);
+				if (NumChunksRequired < MaxChunksPerBatch) //everything fits in a single batch
+				{
+					NumChunksPerBatch                   = NumChunksRequired;
+					//take as little memory as possible and execute it in equal sized chunks
+					MaxLoopsPerChunk                    = (TotalNumLoopsRequired + NumChunksRequired - 1) / NumChunksRequired;
+				}
+				else //not everything fits in a single batch, we have to thread this
+				{
+					MaxLoopsPerChunk                    = (uint32)(NumBytesPerBatchAvailableForTempRegs / NumBytesRequiredPerLoop);
+					uint32 NumLoopsPerBatch             = MaxLoopsPerChunk * NumChunksPerBatch;
+					NumBatches                          = (TotalNumLoopsRequired + NumLoopsPerBatch - 1) / NumLoopsPerBatch;
+					if (GVVMMaxThreadsPerScript > 0 && NumBatches > GVVMMaxThreadsPerScript)
+					{
+						//number of batches exceed the number of threads allowed.. increase the number of chunks per batch
+						NumLoopsPerBatch                = (TotalNumLoopsRequired + GVVMMaxThreadsPerScript - 1) / GVVMMaxThreadsPerScript;
+						NumChunksPerBatch               = (int)(NumLoopsPerBatch + MaxLoopsPerChunk - 1) / MaxLoopsPerChunk;
+						NumBatches                      = GVVMMaxThreadsPerScript;
+						check(NumBatches * NumChunksPerBatch * MaxLoopsPerChunk >= TotalNumLoopsRequired);
+					}
+				}
+			}
+			else
+			{
+				//everything fits into a single chunk
+				NumChunksPerBatch = 1;
+				MaxLoopsPerChunk = TotalNumLoopsRequired;
+			}
+		}
+		PerBatchRegisterDataBytesRequired = ConstantBufferSize + MaxLoopsPerChunk * NumBytesRequiredPerLoop;
+	}
+
+	size_t NumBytesRequiredPerBatch = PerBatchRegisterDataBytesRequired + PerBatchChunkLocalDataOutputIdxBytesRequired + PerBatchChunkLocalNumOutputBytesRequired + PerBatchChunkLocalNumExtFnDecodeRegisterBytesRequired;
+	check(NumBytesRequiredPerBatch <= GVVMChunkSizeInBytes || BatchOverheadSize > GVVMChunkSizeInBytes); //if BatchOverheadSize is too high, then this check is invalid since we're blowing past the limit anyway and there's nothing we can do about it
+
+	{
+		//after the batch size has been calcluated we add two more things: padding and indices for deterministic random generation
+		//This will often cause the actual memory usage to go slightly above the GVVMChunkSizeInBytes.  I decided this is okay, because
+		//the concern is more to get the runtime memory usage lower than the L1 size, not to split hairs over a few hundred bytes.
+		//The random counters are very rarely needed, and are only included for backwards compatiblity with the previous VM.
+		//I figured it'd be better to append these at the end and not mess up what fits in the L1 for the 99.99% of cases that don't require it.
+		//There is no way to determine ahead of time if a script will need the random counters or not.
+		NumBytesRequiredPerBatch += VVM_CACHELINE_SIZE;
+		NumBytesRequiredPerBatch += sizeof(int32) * MaxLoopsPerChunk * 4;
+	}
+
+	{ //compute the number of overhead bytes for this VVM State
+		size_t NumDataSetOutputBytesRequired            = sizeof(volatile long)            * InitData->OptimizeContext->NumOutputDataSets;
+		size_t NumBatchStateBytesRequired               = sizeof(FVectorVMBatchState)      * NumBatches; //this is *NOT* the memory per batch, just the overhead stored in VVMState
+		size_t NumExtFnBytesRequired                    = sizeof(FVectorVMExtFunctionData) * InitData->ExtFunctionTable.Num();
+		NumVVMStateBytesRequired                        = VVM_ALIGN_64(NumDataSetOutputBytesRequired +
+			                                                           NumExtFnBytesRequired +
+			                                                           NumBatchStateBytesRequired);
+	}
+
+	size_t TotalBytesRequired = VVM_ALIGN_64(sizeof(FVectorVMState) + NumVVMStateBytesRequired) + NumBytesRequiredPerBatch; //first batch gets allocated immediately following the VVM State
+
+	VectorVMReallocFn * ReallocFn = InitData->ReallocFn ? InitData->ReallocFn : VVMDefaultRealloc;
+	VectorVMFreeFn *    FreeFn    = InitData->FreeFn    ? InitData->FreeFn    : VVMDefaultFree;
+	
+	FVectorVMState *VVMState;
+	if (InitData->ExistingVectorVMState)
+	{
+		//start at 1 because the first batch is allocated immediately following the VVMState and will be freed when VVMState is freed
+		for (int i = 1; i < InitData->ExistingVectorVMState->NumBatches; ++i)
+		{
+			FreeFn(InitData->ExistingVectorVMState->BatchStates[i].MallocedMemPtr, __FILE__, __LINE__);
+		}
+	}
+	if (InitData->ExistingVectorVMState && InitData->ExistingVectorVMState->NumBytesMalloced >= TotalBytesRequired)
+	{
+		VVMState = InitData->ExistingVectorVMState;
+	}
+	else
+	{
+		VVMState = (FVectorVMState *)ReallocFn(InitData->ExistingVectorVMState, TotalBytesRequired, __FILE__, __LINE__);
+		VVMState->NumBytesMalloced = TotalBytesRequired;
+	}
+	uint8 *VVMStatePtr       = (uint8 *)VVM_ALIGN_16((size_t)(VVMState + 1)); //start state pointers immediately following the VVMState
+	if (VVMState == nullptr)
+	{
+		return nullptr;
+	}
+	
+	VVMState->ReallocFn             = ReallocFn;
+	VVMState->FreeFn                = FreeFn;
+	VVMState->Error.Flags           = 0;
+	VVMState->Error.LineNum         = 0;
+	
+#	define IncVVMStatePtr(NumBytes, ExtraErrFlags)                                      \
+	VVMStatePtr += NumBytes;                                                            \
+	if ((size_t)((uint8 *)VVMStatePtr - (uint8 *)VVMState) > TotalBytesRequired)        \
+	{                                                                                   \
+		VVMState->Error.Flags |= VVMErr_InitMemMismatch | VVMErr_Fatal | ExtraErrFlags; \
+		VVMState->Error.LineNum = __LINE__;                                             \
+		return VVMState;                                                                \
+	}
+
+	{ //we have enough memory malloced for this state, so set up the stuff that comes immediately after the VVMState
+		if (InitData->OptimizeContext->NumExtFns > 0)
+		{
+			VVMState->NumExtFunctions  = InitData->OptimizeContext->NumExtFns;
+			VVMState->ExtFunctionTable = (FVectorVMExtFunctionData *)VVMStatePtr;        IncVVMStatePtr(sizeof(FVectorVMExtFunctionData) * VVMState->NumExtFunctions, 0);
+			for (uint32 i = 0; i < InitData->OptimizeContext->NumExtFns; ++i)
+			{
+				VVMState->ExtFunctionTable[i].Function   = InitData->ExtFunctionTable[i];
+				VVMState->ExtFunctionTable[i].NumInputs  = InitData->OptimizeContext->ExtFnTable[i].NumInputs;
+				VVMState->ExtFunctionTable[i].NumOutputs = InitData->OptimizeContext->ExtFnTable[i].NumOutputs;
+			}
+		}
+		else
+		{
+			VVMState->ExtFunctionTable = nullptr;
+			VVMState->NumExtFunctions  = 0;
+		}
+		uint8 *StateAfterExtFnPtr = VVMStatePtr;
+		VVMState->NumOutputPerDataSet  = (volatile int32 *)VVMStatePtr;                  IncVVMStatePtr(sizeof(volatile int32) * InitData->OptimizeContext->NumOutputDataSets, 0);
+		VVMState->BatchStates          = (FVectorVMBatchState *)VVMStatePtr;             IncVVMStatePtr(sizeof(volatile FVectorVMBatchState) * NumBatches, 0);
+		FMemory::Memset(StateAfterExtFnPtr, 0, (size_t)(VVMStatePtr - StateAfterExtFnPtr));
+	}
+#	undef IncVVMStatePtr
+
+	//init the part of the VVMState that doesn't require the externally allocated memory
+	VVMState->Bytecode          = InitData->OptimizeContext->OutputBytecode;
+	VVMState->NumBytecodeBytes  = InitData->OptimizeContext->NumBytecodeBytes;
+	VVMState->NumTempRegisters  = InitData->OptimizeContext->NumTempRegisters;
+	VVMState->NumConstBuffers   = InitData->OptimizeContext->NumConstsRemapped;
+	VVMState->NumOutputDataSets = InitData->OptimizeContext->NumOutputDataSets;
+	VVMState->MaxExtFnRegisters = MaxExtFnRegisters;
+	VVMState->DataSets          = InitData->DataSets;
+
+	VVMState->UserPtrTable      = InitData->UserPtrTable;
+	VVMState->NumUserPtrTable   = InitData->NumUserPtrTable;
+	VVMState->TotalNumInstances = InitData->NumInstances;
+
+	//batch data
+	VVMState->NumBytesRequiredPerBatch                              = NumBytesRequiredPerBatch;
+	VVMState->PerBatchRegisterDataBytesRequired                     = PerBatchRegisterDataBytesRequired;
+	VVMState->PerBatchChunkLocalDataOutputIdxBytesRequired          = PerBatchChunkLocalDataOutputIdxBytesRequired;
+	VVMState->PerBatchChunkLocalNumOutputBytesRequired              = PerBatchChunkLocalNumOutputBytesRequired;
+	VVMState->PerBatchChunkLocalNumExtFnDecodeRegisterBytesRequired = PerBatchChunkLocalNumExtFnDecodeRegisterBytesRequired;
+
+	VVMState->NumBatches                     = NumBatches;
+	VVMState->MaxChunksPerBatch              = NumChunksPerBatch;
+	VVMState->MaxInstancesPerChunk           = MaxLoopsPerChunk << 2;
+	VVMState->NumInstancesAssignedToBatches  = 0;
+	VVMState->NumInstancesCompleted          = 0;
+
+	if (NumBatches > 0) //init the first batch using the remainder of the memory
+	{
+		FVectorVMBatchState *BatchState = VVMState->BatchStates + 0;
+		VVMStatePtr = SetupBatchStatePtrs(VVMState, BatchState, VVMStatePtr);
+
+		if (VVMStatePtr) //init the constant data at the start of the batch state
+		{
+			//first compute the starting offset of each set of constants
+			uint32 ConstCountAcc = 0;
+			for (int i = 0; i < InitData->NumConstData; ++i)
+			{
+				InitData->ConstData[i].StartingOffset = ConstCountAcc;
+				ConstCountAcc += InitData->ConstData[i].NumDWords;
+			}
+			
+			FVecReg *ConstantBuffers = BatchState->RegisterData;
+			for (uint32 i = 0; i < VVMState->NumConstBuffers; ++i)
+			{
+				uint16 ConstBufferOffset = InitData->OptimizeContext->ConstRemap[1][i];
+				for (int j = 0; j < InitData->NumConstData; ++j)
+				{
+					FVectorVMConstData *ConstData = InitData->ConstData + j;
+					if (ConstBufferOffset >= ConstData->StartingOffset && ConstBufferOffset < ConstData->StartingOffset + ConstData->NumDWords)
+					{
+						ConstantBuffers[i].i = VectorIntSet1(((uint32 *)ConstData->RegisterData)[ConstBufferOffset - ConstData->StartingOffset]);
+						break;
+					}
+				}
+			}
+		}
+		else
+		{
+			return VVMState;
+		}
+		SetupRandStateForBatch(BatchState);
+	}
+	if (SerializeState)
+	{
+		VVMSer_initSerializationState(VVMState, SerializeState, InitData, SerializeState->Flags | VVMSer_OptimizedBytecode);
+	}
+	return VVMState;
+}
+
+
+static void ExecVVMBatch(FVectorVMState *VVMState, int ExecIdx, FVectorVMSerializeState *SerializeState, FVectorVMSerializeState *CmpSerializeState)
+{
+#	if defined(VVM_INCLUDE_SERIALIZATION) && !defined(VVM_SERIALIZE_NO_WRITE)
+#		ifdef VVM_SERIALIZE_PERF
+#			define serializeIns(Type, NumParams)	
+#			define serializeRegUsed(RegIdx, Type)
+#		else //VVM_SERIALIZE_PERF
+#			define serializeIns(Type, NumParams)	if (SerializeState)                                                                         \
+													{                                                                                           \
+														for (int vi = 0; vi <= (int)(NumParams); ++vi)                                          \
+														{                                                                                       \
+															if ((VecIndices[vi] & 0x8000) == 0)											        \
+															{                                                                                   \
+																SerializeState->TempRegFlags[VecIndices[vi]] = VVMRegFlag_Clean + (Type);       \
+															}                                                                                   \
+														}                                                                                       \
+													}
+
+#			define serializeRegUsed(RegIdx, Type)	if (SerializeState)                                                                         \
+													{                                                                                           \
+														SerializeState->TempRegFlags[RegIdx] = VVMRegFlag_Clean + (Type);                       \
+													}
+#		endif //VVM_SERIALIZE_PERF
+#	else //VVM_INCLUDE_SERIALIZATION
+#		define serializeIns(Type, NumParams)		
+#		define serializeRegUsed(RegIdx, Type)
+#	endif //VVM_INCLUDE_SERIALIZATION
+
+#	define execVecIns1f(ins_)						serializeIns(0, 1)                                                  \
+													for (int i = 0; i < NumLoops; ++i)                                  \
+													{                                                                   \
+														VectorRegister4f r0  = VectorLoad(&VecReg[0][i & RegInc[0]].v); \
+														VectorRegister4f res = ins_(r0);                                \
+														VectorStoreAligned(res, &VecReg[1][i].v);                       \
+													}                                                                   \
+													InsPtr += 4;
+
+#	define execVecIns2f(ins_)						serializeIns(0, 2)                                                  \
+													for (int i = 0; i < NumLoops; ++i)                                  \
+													{                                                                   \
+														VectorRegister4f r0  = VectorLoad(&VecReg[0][i & RegInc[0]].v); \
+														VectorRegister4f r1  = VectorLoad(&VecReg[1][i & RegInc[1]].v); \
+														VectorRegister4f res = ins_(r0, r1);                            \
+														VectorStoreAligned(res, &VecReg[2][i].v);                       \
+													}                                                                   \
+													InsPtr += 6;
+
+#	define execVecIns3f(ins_)						serializeIns(0, 3)                                                  \
+													for (int i = 0; i < NumLoops; ++i)                                  \
+													{                                                                   \
+														VectorRegister4f r0  = VectorLoad(&VecReg[0][i & RegInc[0]].v); \
+														VectorRegister4f r1  = VectorLoad(&VecReg[1][i & RegInc[1]].v); \
+														VectorRegister4f r2  = VectorLoad(&VecReg[2][i & RegInc[2]].v); \
+														VectorRegister4f res = ins_(r0, r1, r2);                        \
+														VectorStoreAligned(res, &VecReg[3][i].v);                       \
+													}                                                                   \
+													InsPtr += 8;
+
+#	define execVecIns1i(ins_)						serializeIns(1, 1)                                                     \
+													for (int i = 0; i < NumLoops; ++i)                                     \
+													{                                                                      \
+														VectorRegister4i r0  = VectorIntLoad(&VecReg[0][i & RegInc[0]].v); \
+														VectorRegister4i res = ins_(r0);                                   \
+														VectorIntStoreAligned(res, &VecReg[1][i].i);                       \
+													}                                                                      \
+													InsPtr += 4;
+
+#	define execVecIns2i(ins_)						serializeIns(1, 2)                                                     \
+													for (int i = 0; i < NumLoops; ++i)                                     \
+													{                                                                      \
+														VectorRegister4i r0  = VectorIntLoad(&VecReg[0][i & RegInc[0]].v); \
+														VectorRegister4i r1  = VectorIntLoad(&VecReg[1][i & RegInc[1]].v); \
+														VectorRegister4i res = ins_(r0, r1);                               \
+														VectorIntStoreAligned(res, &VecReg[2][i].i);                       \
+													}                                                                      \
+													InsPtr += 6;
+
+#	define execVecIns3i(ins_)						serializeIns(1, 3)                                                     \
+													for (int i = 0; i < NumLoops; ++i)                                     \
+													{                                                                      \
+														VectorRegister4i r0  = VectorIntLoad(&VecReg[0][i & RegInc[0]].v); \
+														VectorRegister4i r1  = VectorIntLoad(&VecReg[1][i & RegInc[1]].v); \
+														VectorRegister4i r2  = VectorIntLoad(&VecReg[2][i & RegInc[2]].v); \
+														VectorRegister4i res = ins_(r0, r1, r2);                           \
+														VectorIntStoreAligned(res, &VecReg[3][i].i);                       \
+													}                                                                      \
+													InsPtr += 8;
+
+
+	FVectorVMBatchState *BatchState = nullptr;
+	int BatchIdx = -1;
+	//check to see if we can reuse a batch's memory that's finished executing
+	for (int i = 0; i < VVMState->NumBatches; ++i)
+	{
+		int32 WasCurrentlyExecuting = FPlatformAtomics::InterlockedCompareExchange(&VVMState->BatchStates[i].CurrentlyExecuting, 1, 0);
+		if (WasCurrentlyExecuting == 0)
+		{
+			//we can reuse this batch
+			BatchState = VVMState->BatchStates + i;
+			BatchIdx = i;
+			break;
+		}
+	}
+	if (BatchIdx == -1 || BatchState == nullptr)
+	{
+		check(false);
+		return;
+	}
+	if (BatchState->RegisterData == nullptr)
+	{
+		check(BatchIdx != 0); //first batch state should have set the pointers in init()
+		//this is the first time using this batch, so we need to malloc the data and copy the consts over from batch state 0
+		check(VVMState->NumBytesRequiredPerBatch >= VVMState->PerBatchRegisterDataBytesRequired +
+		                                            VVMState->PerBatchChunkLocalDataOutputIdxBytesRequired +
+		                                            VVMState->PerBatchChunkLocalNumOutputBytesRequired +
+			                                        VVM_CACHELINE_SIZE /* padding */);
+		void *BatchData = VVMState->ReallocFn(nullptr, VVMState->NumBytesRequiredPerBatch, __FILE__, __LINE__);
+		if (BatchData == nullptr)
+		{
+			check(false);
+			VVMState->Error.Flags |= VVMErr_BatchMemory;
+			VVMState->Error.LineNum = __LINE__;
+			return;
+		}
+		if (SetupBatchStatePtrs(VVMState, BatchState, BatchData) == nullptr)
+		{
+			return;
+		}
+		VVMMemCpy(BatchState->RegisterData, VVMState->BatchStates[0].RegisterData, sizeof(FVecReg) * VVMState->NumConstBuffers); //copy the constant data from the first batch
+		SetupRandStateForBatch(BatchState);
+	}
+	if (!AssignInstancesToBatch(VVMState, BatchState))
+	{
+		return; //no more instances to do, we're done
+	}
+
+	int StartInstanceThisChunk = BatchState->StartInstance;
+	int NumChunksThisBatch     = (BatchState->NumInstances + VVMState->MaxInstancesPerChunk - 1) / VVMState->MaxInstancesPerChunk;
+	
+	uint32 RegInc[4];
+		
+	VectorRegister4i NumConsts4 = VectorIntSet1(VVMState->NumConstBuffers);
+	
+	for (int ChunkIdxThisBatch = 0; ChunkIdxThisBatch < NumChunksThisBatch; ++ChunkIdxThisBatch, StartInstanceThisChunk += VVMState->MaxInstancesPerChunk)
+	{
+		int NumInstancesThisChunk = VVM_MIN(VVMState->MaxInstancesPerChunk, BatchState->StartInstance + BatchState->NumInstances - StartInstanceThisChunk);
+		VVMSer_chunkStartExp(SerializeState, ChunkIdxThisBatch, BatchIdx);
+		
+		int NumLoops               = (int)((NumInstancesThisChunk + 3) & ~3) >> 2; //assumes 4-wide ops
+		VectorRegister4i NumLoops4 = VectorIntSet1(NumLoops);
+		uint8 *InsPtr              = VVMState->Bytecode;
+		uint8 *InsPtrEnd           = InsPtr + VVMState->NumBytecodeBytes;
+
+		for (uint32 i = 0; i < VVMState->NumOutputDataSets; ++i)
+		{
+			BatchState->ChunkLocalData.StartingOutputIdxPerDataSet[i] = 0;
+			BatchState->ChunkLocalData.NumOutputPerDataSet[i] = 0;
+		}
+		
+#		define VVMDecodeInstructionRegisters(Bytecode, RegData, RegIncMask) do {                                                                                                                                                            \
+			uint32 VecOffsets[4];                                                                                                                                                                                                           \
+			VectorRegister4i VecIndicesIn4  = VectorIntLoad(Bytecode);                                                                          /* 16 bit inputs.  15 bits for index, 1 high bit for const/reg flag (0: reg, 1: const) */   \
+			VectorRegister4i VecIndices4    = VectorIntExpandLow16To32(VecIndicesIn4);                                                          /* 4-wide 32 bit version of the inputs, bits 16:31 are 0 */                                 \
+			VectorRegister4i RegIncInv4     = VectorIntSubtract(VectorSetZero(), VectorShiftRightImmArithmetic(VecIndices4, 15));               /* Sets the inverse of what we need: 0xFF... is const, 0 is reg */                          \
+			VectorRegister4i RegInc4        = VectorIntXor(RegIncInv4, VVM_m128iConst(FMask));                                                  /* Whether to increment the index counter: 0xFF... for registers, and 0 for const */        \
+			VectorRegister4i VecRegIndices4 = VectorIntAnd(VecIndices4, VVM_m128iConst(RegOffsetMask));                                         /* Only the register index, the const/reg flag is stripped off */                           \
+			VectorRegister4i ConstOffset4   = VectorIntAnd(VecRegIndices4, RegIncInv4);                                                         /* Only the const offsets, all registers are masked out */                                  \
+			VectorRegister4i TempRegOffset4 = VectorIntAnd(VectorIntAdd(NumConsts4, VectorIntMultiply(VecRegIndices4, NumLoops4)), RegInc4);    /* Only the register offsets, all consts are masked out */                                  \
+			VectorRegister4i OptRegOffset4  = VectorIntOr(ConstOffset4, TempRegOffset4);                                                        /* Blended (sse4 would be nice) const and temp register offsets */                          \
+			VectorIntStore(RegInc4, RegIncMask);                                                                                                                                                                                            \
+			VectorIntStore(OptRegOffset4, VecOffsets);                                                                                                                                                                                      \
+			(RegData)[0] = BatchState->RegisterData + VecOffsets[0];                                                                                                                                                                        \
+			(RegData)[1] = BatchState->RegisterData + VecOffsets[1];                                                                                                                                                                        \
+			(RegData)[2] = BatchState->RegisterData + VecOffsets[2];                                                                                                                                                                        \
+			(RegData)[3] = BatchState->RegisterData + VecOffsets[3];                                                                                                                                                                        \
+		} while (0)
+
+		FVecReg *VecReg[4];
+		
+		while (InsPtr < InsPtrEnd)
+		{
+			VVMSer_insStartExp(SerializeState);
+			EVectorVMOp OpCode = (EVectorVMOp)*InsPtr++;
+			uint16 *VecIndices = (uint16 *)InsPtr;
+			VVMDecodeInstructionRegisters(InsPtr, VecReg, RegInc);
+			VVMSer_insEndDecodeExp(SerializeState);
+
+OpCodeSwitch: //I think computed gotos would be a huge win here... maybe write this loop in assembly for the jump table?!
+			switch (OpCode) {
+				case EVectorVMOp::done:                                                                 break;
+				case EVectorVMOp::add:                              execVecIns2f(VectorAdd);            break;
+				case EVectorVMOp::sub:                              execVecIns2f(VectorSubtract);       break;
+				case EVectorVMOp::mul:                              execVecIns2f(VectorMultiply);       break;
+				case EVectorVMOp::div:                              execVecIns2f(VVM_safeIns_div);      break;
+				case EVectorVMOp::mad:                              execVecIns3f(VectorMultiplyAdd);    break;
+				case EVectorVMOp::lerp:                             execVecIns3f(VectorLerp);           break;
+				case EVectorVMOp::rcp:                              execVecIns1f(VVM_safeIns_rcp);      break;
+				case EVectorVMOp::rsq:                              execVecIns1f(VVM_safe_rsq);         break;
+				case EVectorVMOp::sqrt:                             execVecIns1f(VVM_safe_sqrt);        break;
+				case EVectorVMOp::neg:                              execVecIns1f(VectorNegate);         break;
+				case EVectorVMOp::abs:                              execVecIns1f(VectorAbs);            break;
+				case EVectorVMOp::exp:                              execVecIns1f(VectorExp);            break;
+				case EVectorVMOp::exp2:                             execVecIns1f(VectorExp2);           break;
+				case EVectorVMOp::log:                              execVecIns1f(VVM_safe_log);         break;
+				case EVectorVMOp::log2:                             execVecIns1f(VectorLog2);           break;
+				case EVectorVMOp::sin:                              execVecIns1f(VectorSin);            break;
+				case EVectorVMOp::cos:                              execVecIns1f(VectorCos);            break;
+				case EVectorVMOp::tan:                              execVecIns1f(VectorTan);            break;
+				case EVectorVMOp::asin:                             execVecIns1f(VectorASin);           break;
+				case EVectorVMOp::acos:                             execVecIns1f(VVM_vecACosFast);      break;
+				case EVectorVMOp::atan:                             execVecIns1f(VectorATan);           break;
+				case EVectorVMOp::atan2:                            execVecIns2f(VectorATan2);          break;
+				case EVectorVMOp::ceil:                             execVecIns1f(VectorCeil);           break;
+				case EVectorVMOp::floor:                            execVecIns1f(VectorFloor);          break;
+				case EVectorVMOp::fmod:                             execVecIns2f(VectorMod);            break;
+				case EVectorVMOp::frac:                             execVecIns1f(VectorFractional);     break;
+				case EVectorVMOp::trunc:                            execVecIns1f(VectorTruncate);       break;
+				case EVectorVMOp::clamp:                            execVecIns3f(VectorClamp);          break;
+				case EVectorVMOp::min:                              execVecIns2f(VectorMin);            break;
+				case EVectorVMOp::max:                              execVecIns2f(VectorMax);            break;
+				case EVectorVMOp::pow:                              execVecIns2f(VVM_safe_pow);         break;
+				case EVectorVMOp::round:                            execVecIns1f(VectorRound);          break;
+				case EVectorVMOp::sign:                             execVecIns1f(VectorSign);           break;
+				case EVectorVMOp::step:                             execVecIns2f(VVM_vecStep);          break;
+				case EVectorVMOp::random:
+#					if defined(VVM_INCLUDE_SERIALIZATION) && !defined(VVM_SERIALIZE_NO_WRITE)
+#					ifndef VVM_SERIALIZE_PERF
+					if (SerializeState && (SerializeState->Flags & VVMSer_SyncRandom) && CmpSerializeState && CmpSerializeState->NumInstructions >= SerializeState->NumInstructions && CmpSerializeState->NumTempRegisters > VecIndices[1])
+					{
+						FVectorVMSerializeInstruction *CmpIns = nullptr;
+						if ((SerializeState->Flags & VVMSer_OptimizedBytecode) && SerializeState->OptimizeCtx && !(CmpSerializeState->Flags & VVMSer_OptimizedBytecode))
+						{
+							for (uint32 i = 0; i < SerializeState->OptimizeCtx->Intermediate.NumInstructions; ++i)
+							{
+								if (SerializeState->OptimizeCtx->Intermediate.Instructions[i].PtrOffsetInOptimizedBytecode + 1 == (uint32)(InsPtr - VVMState->Bytecode))
+								{ //+1 because we already incremented InsPtr
+									for (uint32 j = 0; j < CmpSerializeState->NumInstructions; ++j)
+									{
+										if (CmpSerializeState->Instructions[j].OpStart == SerializeState->OptimizeCtx->Intermediate.Instructions[i].PtrOffsetInOrigBytecode)
+										{
+											CmpIns = CmpSerializeState->Instructions + j;
+											break;
+										}
+									}
+									break;
+								}
+							}
+						}
+						else
+						{
+							CmpIns = CmpSerializeState->Instructions + SerializeState->NumInstructions + 1; //this instruction
+						}
+						if (CmpIns)
+						{
+							check(CmpSerializeState->Bytecode[CmpIns->OpStart] == (uint8)EVectorVMOp::random);
+							uint16 CmpVecIdx = ((uint16 *)(CmpSerializeState->Bytecode + CmpIns->OpStart + 2))[1];
+							check((CmpVecIdx & 8000) == 0); //can't output to constant
+							float *op_reg = (float *)VecReg[1];
+							float *ip_reg = (float *)CmpIns->TempRegisters + SerializeState->NumInstances * CmpVecIdx + StartInstanceThisChunk;
+							for (int i = 0; i < NumInstancesThisChunk; ++i)
+							{
+								*op_reg++ = *ip_reg++;
+							}
+							goto VVM_EndSyncRandom;
+						}
+					}
+#					endif //VVM_SERIALIZE_PERF
+#					endif //VVM_INCLUDE_SERIALIZATION
+					for (int i = 0; i < NumLoops; ++i)
+					{
+						VectorRegister4i RandReg = VVMXorwowStep(BatchState);
+						VectorRegister4i IntPart = VectorIntOr(VectorShiftRightImmLogical(RandReg, 9), VectorIntSet1(0x3F800000));
+						VectorRegister4f FltPart = VectorCastIntToFloat(IntPart);
+						VecReg[1][i].v = VectorMultiply(VectorSubtract(FltPart, VectorSet1(1.f)), VecReg[0][i & RegInc[0]].v);
+					}
+#					if defined(VVM_INCLUDE_SERIALIZATION) && !defined(VVM_SERIALIZE_NO_WRITE)
+					VVM_EndSyncRandom:
+#					endif //VVM_INCLUDE_SERIALIZATION
+					serializeRegUsed(VecIndices[1], 0);
+					InsPtr += 4;
+				break;
+				case EVectorVMOp::noise:                            check(false);                       break;
+				case EVectorVMOp::cmplt:                            execVecIns2f(VectorCompareLT);      break;
+				case EVectorVMOp::cmple:                            execVecIns2f(VectorCompareLE);      break;
+				case EVectorVMOp::cmpgt:                            execVecIns2f(VectorCompareGT);      break;
+				case EVectorVMOp::cmpge:                            execVecIns2f(VectorCompareGE);      break;
+				case EVectorVMOp::cmpeq:                            execVecIns2f(VectorCompareEQ);      break;
+				case EVectorVMOp::cmpneq:                           execVecIns2f(VectorCompareNE);      break;
+				case EVectorVMOp::select:                           execVecIns3i(VectorIntSelect);      break;
+				case EVectorVMOp::addi:                             execVecIns2i(VectorIntAdd);         break;
+				case EVectorVMOp::subi:                             execVecIns2i(VectorIntSubtract);    break;
+				case EVectorVMOp::muli:                             execVecIns2i(VectorIntMultiply);    break;
+				case EVectorVMOp::divi: {
+					serializeIns(1, 2)
+					//@TODO: convert to double and div 4 wide
+					for (int i = 0; i < NumLoops; ++i)
+					{
+						int32 TmpA[4];
+						VectorIntStore(VecReg[0][i & RegInc[0]].i, TmpA);
+						int32 TmpB[4];
+						VectorIntStore(VecReg[1][i & RegInc[1]].i, TmpB);
+
+						// No intrinsics exist for integer divide. Since div by zero causes crashes, we must be safe against that.
+						int32 TmpDst[4];
+						TmpDst[0] = TmpB[0] != 0 ? (TmpA[0] / TmpB[0]) : 0;
+						TmpDst[1] = TmpB[1] != 0 ? (TmpA[1] / TmpB[1]) : 0;
+						TmpDst[2] = TmpB[2] != 0 ? (TmpA[2] / TmpB[2]) : 0;
+						TmpDst[3] = TmpB[3] != 0 ? (TmpA[3] / TmpB[3]) : 0;
+						VecReg[2][i].i = MakeVectorRegisterInt(TmpDst[0], TmpDst[1], TmpDst[2], TmpDst[3]);
+					}
+					InsPtr += 6;
+				} break;
+				case EVectorVMOp::clampi:                           execVecIns3i(VectorIntClamp);       break;
+				case EVectorVMOp::mini:                             execVecIns2i(VectorIntMin);         break;
+				case EVectorVMOp::maxi:                             execVecIns2i(VectorIntMax);         break;
+				case EVectorVMOp::absi:                             execVecIns1i(VectorIntAbs);         break;
+				case EVectorVMOp::negi:                             execVecIns1i(VectorIntNegate);      break;
+				case EVectorVMOp::signi:                            execVecIns1i(VectorIntSign);        break;
+				case EVectorVMOp::randomi: {
+					//@TODO: serialize syncing, no test cases yet
+					serializeIns(1, 1)
+					for (int i = 0; i < NumLoops; ++i)
+					{
+						VecReg[0][i].i = VVMXorwowStep(BatchState);
+					}
+					InsPtr += 4;
+				} break;
+				case EVectorVMOp::cmplti:                           execVecIns2i(VectorIntCompareLT);   break;
+				case EVectorVMOp::cmplei:                           execVecIns2i(VectorIntCompareLE);   break;
+				case EVectorVMOp::cmpgti:                           execVecIns2i(VectorIntCompareGT);   break;
+				case EVectorVMOp::cmpgei:                           execVecIns2i(VectorIntCompareGE);   break;
+				case EVectorVMOp::cmpeqi:                           execVecIns2i(VectorIntCompareEQ);   break;
+				case EVectorVMOp::cmpneqi:                          execVecIns2i(VectorIntCompareNEQ);  break;
+				case EVectorVMOp::bit_and:                          execVecIns2i(VectorIntAnd);         break;
+				case EVectorVMOp::bit_or:                           execVecIns2i(VectorIntOr);          break;
+				case EVectorVMOp::bit_xor:                          execVecIns2i(VectorIntXor);         break;
+				case EVectorVMOp::bit_not:                          execVecIns1i(VectorIntNot);         break;
+				case EVectorVMOp::bit_lshift:
+					for (int i = 0; i < NumLoops; ++i)
+					{
+						int idx0 = (i << 2) & RegInc[0];
+						int idx1 = (i << 2) & RegInc[1];
+
+						((int *)VecReg[2])[(i << 2) + 0] = ((int *)VecReg[0])[idx0 + 0] << ((int *)VecReg[1])[idx1 + 0];
+						((int *)VecReg[2])[(i << 2) + 1] = ((int *)VecReg[0])[idx0 + 1] << ((int *)VecReg[1])[idx1 + 1];
+						((int *)VecReg[2])[(i << 2) + 2] = ((int *)VecReg[0])[idx0 + 2] << ((int *)VecReg[1])[idx1 + 2];
+						((int *)VecReg[2])[(i << 2) + 3] = ((int *)VecReg[0])[idx0 + 3] << ((int *)VecReg[1])[idx1 + 3];
+					}
+					serializeRegUsed(VecIndices[2], 1);
+					InsPtr += 6;
+					break;
+				case EVectorVMOp::bit_rshift:
+					for (int i = 0; i < NumLoops; ++i)
+					{
+						int idx0 = (i << 2) & RegInc[0];
+						int idx1 = (i << 2) & RegInc[1];
+						((int *)VecReg[2])[(i << 2) + 0] = ((int *)VecReg[0])[idx0 + 0] >> ((int *)VecReg[1])[idx1 + 0];
+						((int *)VecReg[2])[(i << 2) + 1] = ((int *)VecReg[0])[idx0 + 1] >> ((int *)VecReg[1])[idx1 + 1];
+						((int *)VecReg[2])[(i << 2) + 2] = ((int *)VecReg[0])[idx0 + 2] >> ((int *)VecReg[1])[idx1 + 2];
+						((int *)VecReg[2])[(i << 2) + 3] = ((int *)VecReg[0])[idx0 + 3] >> ((int *)VecReg[1])[idx1 + 3];
+					}
+					serializeRegUsed(VecIndices[2], 1);
+					InsPtr += 6;
+					break;
+				case EVectorVMOp::logic_and:                        execVecIns2i(VectorIntAnd);         break;
+				case EVectorVMOp::logic_or:                         execVecIns2i(VectorIntOr);          break;
+				case EVectorVMOp::logic_xor:                        execVecIns2i(VectorIntXor);         break;
+				case EVectorVMOp::logic_not:                        execVecIns1i(VectorIntNot);         break;
+				case EVectorVMOp::f2i:
+					for (int i = 0; i < NumLoops; ++i)
+					{
+						VecReg[1][i].i = VectorFloatToInt(VecReg[0][i & RegInc[0]].v);
+					}
+					serializeRegUsed(VecIndices[1], 1);
+					InsPtr += 4;
+				break;
+				case EVectorVMOp::i2f:
+					for (int i = 0; i < NumLoops; ++i)
+					{
+						VecReg[1][i].v = VectorIntToFloat(VecReg[0][i & RegInc[0]].i);
+					}
+					serializeRegUsed(VecIndices[1], 0);
+					InsPtr += 4;
+				break;
+				case EVectorVMOp::f2b:                              execVecIns1f(VVM_vecFloatToBool);   break;
+				case EVectorVMOp::b2f:                              execVecIns1f(VVM_vecBoolToFloat);   break;
+				case EVectorVMOp::i2b:                              execVecIns1i(VVM_vecIntToBool);     break;
+				case EVectorVMOp::b2i:                              execVecIns1i(VVM_vecBoolToInt);     break;
+				case EVectorVMOp::inputdata_float:
+				case EVectorVMOp::inputdata_int32: {
+					uint8 RegType             = (uint8)OpCode - (uint8)EVectorVMOp::inputdata_float;
+					uint16 DataSetIdx         = *(uint16 *)(InsPtr    );
+					uint16 InputRegIdx        = *(uint16 *)(InsPtr + 2);
+					uint16 DestRegIdx         = *(uint16 *)(InsPtr + 4);
+					uint32 InstanceOffset     = VVMState->DataSets[DataSetIdx].InstanceOffset;
+					uint32 InputRegTypeOffset = VVMState->DataSets[DataSetIdx].InputRegisterTypeOffsets[RegType];
+					size_t DstIdx             = NumLoops * DestRegIdx;
+					uint32 **InputBuffers     = (uint32 **)VVMState->DataSets[DataSetIdx].InputRegisters.GetData();
+					VVMMemCpy(BatchState->RegisterData + VVMState->NumConstBuffers + DstIdx, InputBuffers[InputRegIdx + InputRegTypeOffset] + StartInstanceThisChunk + InstanceOffset, sizeof(FVecReg) * NumLoops);
+					serializeRegUsed(DestRegIdx, RegType);
+					InsPtr += 6;
+				} break;
+				case EVectorVMOp::inputdata_half:                   check(false);                       break;
+				case EVectorVMOp::inputdata_noadvance_float:
+				case EVectorVMOp::inputdata_noadvance_int32: {
+					uint8 RegType                 = (uint8)OpCode - (uint8)EVectorVMOp::inputdata_noadvance_float;
+					uint16 DataSetIdx             = *(uint16 *)(InsPtr    );
+					uint16 InputRegIdx            = *(uint16 *)(InsPtr + 2);
+					uint16 DestRegIdx             = *(uint16 *)(InsPtr + 4);
+					uint32 InstanceOffset         = VVMState->DataSets[DataSetIdx].InstanceOffset;
+					uint32 InputRegTypeOffset     = VVMState->DataSets[DataSetIdx].InputRegisterTypeOffsets[RegType];
+					size_t DstIdx                 = NumLoops * DestRegIdx;
+					uint32 **InputBuffers         = (uint32 **)VVMState->DataSets[DataSetIdx].InputRegisters.GetData();
+					uint32 *InputBuffer           = InputBuffers[InputRegIdx + InputRegTypeOffset] + InstanceOffset;
+					VectorRegister4i input_val4   = VectorIntSet1(*InputBuffer);
+					for (int i = 0; i < NumLoops; ++i)
+					{
+						VectorIntStoreAligned(input_val4, BatchState->RegisterData + VVMState->NumConstBuffers + DstIdx + i);
+					}
+					serializeRegUsed(DestRegIdx, RegType);
+					InsPtr += 6;
+				} break;
+				case EVectorVMOp::inputdata_noadvance_half:			check(false);				        break;
+				case EVectorVMOp::outputdata_float:
+				case EVectorVMOp::outputdata_int32: {
+					uint8 RegType = InsPtr[-1] - (uint8)EVectorVMOp::outputdata_float;
+					check(RegType == 0 || RegType == 1); //float or int32
+					uint16 DataSetIdx         = VecIndices[0];
+					int *DstIdxReg            = (int *)VecReg[1];
+					uint32 NumOutputInstances = BatchState->ChunkLocalData.NumOutputPerDataSet[DataSetIdx];
+					uint32 InstanceOffset     = BatchState->ChunkLocalData.StartingOutputIdxPerDataSet[DataSetIdx];
+					uint32 RegTypeOffset      = VVMState->DataSets[DataSetIdx].OutputRegisterTypeOffsets[RegType];
+					uint32 **OutputBuffers    = (uint32 **)VVMState->DataSets[DataSetIdx].OutputRegisters.GetData();
+					uint32 *DstReg            = OutputBuffers[RegTypeOffset + VecIndices[3]] + InstanceOffset;
+					uint32 *SrcReg            = (uint32 *)VecReg[2];
+					if (NumOutputInstances == NumInstancesThisChunk)
+					{
+						if (RegInc[2] == 0) //setting from a constant
+						{ 
+							VVMMemSet32(DstReg, *SrcReg, NumOutputInstances);
+						}
+						else //copying from internal buffers
+						{
+							VVMMemCpy(DstReg, SrcReg, sizeof(uint32) * NumOutputInstances);
+						}
+					}
+					else //if we are discarding at least one instance we can't just copy the memory
+					{
+						for (uint32 i = 0; i < NumOutputInstances; ++i)
+						{
+							DstReg[i] = SrcReg[DstIdxReg[i & RegInc[1]] & RegInc[2]];
+						}
+					}
+					InsPtr += 8;
+				} break;
+				case EVectorVMOp::outputdata_half:					check(false);				        break;
+				case EVectorVMOp::acquireindex:
+				{
+					uint32 NumOutputInstances = 0;
+					uint16 DataSetIdx         = VecIndices[0];
+					uint32 *Input             = (uint32 *)VecReg[1];
+					uint32 *Output            = (uint32 *)VecReg[2];
+					uint16 IncMask            = RegInc[1];
+					for (uint16 i = 0; i < NumInstancesThisChunk; ++i)
+					{
+						//since input and output can alias now we need to save the IncAmt to a temp value
+						uint32 IncAmt = Input[i & IncMask] >> 31; //-1 is keep, so we only need to check for the high bit
+						Output[NumOutputInstances] = i;
+						NumOutputInstances += IncAmt;
+					}
+					//the new VM's indicies are generated to support brachless write-gather for the output instructions (instead of an in-signal flag as the original bytecode intended)
+					//the above loop will write an invalid value into the last slot if we discard one or more instances.  This is normally okay, however if an update_id instruction is issued later, 
+					//we will write incorrect values into the free id table there.  To avoid this (and potentially other problems that may come up if the bytecode is expanded, we correct the final slot here.
+					if ((int)NumOutputInstances < NumInstancesThisChunk)
+					{
+						Output[NumOutputInstances] = NumInstancesThisChunk;
+					}
+					BatchState->ChunkLocalData.StartingOutputIdxPerDataSet[DataSetIdx] = VVMState->DataSets[DataSetIdx].InstanceOffset + FPlatformAtomics::InterlockedAdd(VVMState->NumOutputPerDataSet + DataSetIdx, NumOutputInstances);
+					BatchState->ChunkLocalData.NumOutputPerDataSet[DataSetIdx] += NumOutputInstances;
+					serializeRegUsed(VecIndices[2], VVMRegFlag_Int | VVMRegFlag_Index);
+					InsPtr += 8;
+				}
+				break;
+				case EVectorVMOp::external_func_call:
+				{
+					FVectorVMExtFunctionData *ExtFnData = VVMState->ExtFunctionTable + *InsPtr;
+#if					defined(VVM_INCLUDE_SERIALIZATION) && !defined(VVM_SERIALIZE_NO_WRITE)
+					if (SerializeState && (SerializeState->Flags & VVMSer_SyncExtFns) && CmpSerializeState && SerializeState->NumInstances == CmpSerializeState->NumInstances && (CmpSerializeState->NumInstructions > SerializeState->NumInstructions || VVMSerGlobalChunkIdx != 0))
+					{
+						//If we hit this branch we are using the output from the comparision state instead of running the external function itself.
+						//AFAIK the VM is not speced to have the inputs and outputs in a particular order, and even if it is we shouldn't rely on
+						//3rd party external function writers to follow the spec.  Therefore we don't just sync what we think is output, we sync
+						//all temp registers that are used in the function.
+						int NumRegisters = ExtFnData->NumInputs + ExtFnData->NumOutputs;
+						FVectorVMSerializeInstruction *CmpIns = nullptr; //this instruction
+						FVectorVMOptimizeInstruction *OptIns = nullptr;
+						if (SerializeState->OptimizeCtx)
+						{
+							//instructions have been re-ordered, can't binary search, must linear search
+							for (uint32 i = 0 ; i < SerializeState->OptimizeCtx->Intermediate.NumInstructions; ++i)
+							{
+								if (SerializeState->OptimizeCtx->Intermediate.Instructions[i].PtrOffsetInOptimizedBytecode == (int)(VVMSerStartOpPtr - VVMState->Bytecode))
+								{
+									OptIns = SerializeState->OptimizeCtx->Intermediate.Instructions + i;
+									break;
+								}
+							}
+							if (OptIns)
+							{
+								for (uint32 i = 0; i < CmpSerializeState->NumInstructions; ++i)
+								{
+									if (OptIns->PtrOffsetInOrigBytecode == CmpSerializeState->Instructions[i].OpStart)
+									{
+										CmpIns = CmpSerializeState->Instructions + i;
+										break;
+									}
+								}
+							}
+						}
+						else
+						{
+							CmpIns = CmpSerializeState->Instructions + VVMSerNumInstructionsThisChunk;
+						}
+						if (OptIns && CmpIns)
+						{
+							for (int i = 0; i < NumRegisters; ++i)
+							{
+								uint16 DstRegIdx = ((uint16 *)(InsPtr + 2))[i];
+								if (!(DstRegIdx & 0x8000)) //high bit signifies constant, skip it
+								{
+									uint32 DstOffset = VVMState->NumConstBuffers + DstRegIdx * NumLoops;
+									uint32 SrcRegIdx = DstRegIdx;
+									if (OptIns && (SerializeState->Flags & VVMSer_OptimizedBytecode) && !(CmpSerializeState->Flags & VVMSer_OptimizedBytecode))
+									{
+										SrcRegIdx = ((uint16 *)(CmpSerializeState->Bytecode + OptIns->PtrOffsetInOrigBytecode + 2))[i] & 0x7FFF; //high bit is register in original bytecode
+									}
+									if (SrcRegIdx != 0x7FFF) //invalid register, skipped by external function in the execution
+									{ 
+										uint8 *Src = (uint8 *)(CmpIns->TempRegisters + SrcRegIdx * CmpSerializeState->NumInstances + StartInstanceThisChunk);
+										uint8 *Dst = (uint8 *)(BatchState->RegisterData + DstOffset);
+										VVMMemCpy(Dst, Src, sizeof(uint32) * NumInstancesThisChunk);
+									}
+								}
+							}
+						}
+						for (int i = 0; i < ExtFnData->NumOutputs; ++i)
+						{
+							uint16 RegIdx = ((uint16 *)(InsPtr + 2))[ExtFnData->NumInputs + i];
+							serializeRegUsed(RegIdx, 0); //assume float
+						}
+					} else {
+#					else //VVM_INCLUDE_SERIALIZATION
+					{
+#					endif //VVM_INCLUDE_SERIALIZATION
+						check(*InsPtr < VVMState->NumExtFunctions);
+						check((uint32)(ExtFnData->NumInputs + ExtFnData->NumOutputs) <= VVMState->MaxExtFnRegisters);
+
+						//first decode all of the registers this external function needs into the batch's Chunk Local Data
+						//skip the first index because it's the ExtFnIdx
+						BatchState->ChunkLocalData.ExtFnDecodedReg.RegData[0]    = VecReg[1];
+						BatchState->ChunkLocalData.ExtFnDecodedReg.RegData[1]    = VecReg[2];
+						BatchState->ChunkLocalData.ExtFnDecodedReg.RegData[2]    = VecReg[3];
+
+						BatchState->ChunkLocalData.ExtFnDecodedReg.RegInc[0]     = RegInc[1];
+						BatchState->ChunkLocalData.ExtFnDecodedReg.RegInc[1]     = RegInc[2];
+						BatchState->ChunkLocalData.ExtFnDecodedReg.RegInc[2]     = RegInc[3];
+						
+						for (int i = 3; i < ExtFnData->NumInputs + ExtFnData->NumOutputs; i += 4)
+						{
+							VVMDecodeInstructionRegisters(InsPtr + 2 + i * 2, BatchState->ChunkLocalData.ExtFnDecodedReg.RegData + i, BatchState->ChunkLocalData.ExtFnDecodedReg.RegInc + i);
+						}
+						FVectorVMExternalFunctionContext ExtFnCtx;
+
+						ExtFnCtx.RegisterData             = (uint32 **)BatchState->ChunkLocalData.ExtFnDecodedReg.RegData;
+						ExtFnCtx.RegInc                   = BatchState->ChunkLocalData.ExtFnDecodedReg.RegInc;
+						ExtFnCtx.RawVecIndices            = VecIndices + 1; //skip index 0, that's the function index
+
+						ExtFnCtx.RegReadCount             = 0;
+						ExtFnCtx.NumRegisters             = ExtFnData->NumInputs + ExtFnData->NumOutputs;
+
+						ExtFnCtx.StartInstance            = StartInstanceThisChunk;
+						ExtFnCtx.NumInstances             = NumInstancesThisChunk;
+						ExtFnCtx.NumLoops                 = NumLoops;
+						ExtFnCtx.PerInstanceFnInstanceIdx = 0;
+
+						ExtFnCtx.UserPtrTable             = VVMState->UserPtrTable;
+						ExtFnCtx.NumUserPtrs              = VVMState->NumUserPtrTable;
+						ExtFnCtx.RandStream               = &BatchState->RandStream;
+
+						ExtFnCtx.RandCounters             = BatchState->ChunkLocalData.RandCounters;
+						ExtFnCtx.DataSets                 = VVMState->DataSets;
+						ExtFnData->Function->Execute(ExtFnCtx);
+
+						//UE_LOG(LogVectorVM, Warning, TEXT("Num Instances: %d"), NumInstancesThisChunk);
+					}
+					InsPtr += 2 + ((ExtFnData->NumInputs + ExtFnData->NumOutputs) << 1);
+				}
+				break;
+				case EVectorVMOp::exec_index:
+				{
+					int RegIdx = NumLoops * *((uint16 *)InsPtr);
+					serializeIns(1, 0);
+					VectorRegister4i StartInstance4 = VectorIntAdd(VectorIntSet1(StartInstanceThisChunk), VVM_m128iConst(ZeroOneTwoThree));
+					for (int i = 0; i < NumLoops; ++i)
+					{
+						VectorRegister4i i4 = VectorIntSet1(i);
+						VectorRegister4i v4 = VectorIntAdd(StartInstance4, VectorShiftLeftImm(i4, 2));
+						VecReg[0][i].i = v4;
+					}
+					InsPtr += 2;
+				}
+				break;
+				case EVectorVMOp::noise2D:							check(false);                   break;
+				case EVectorVMOp::noise3D:							check(false);                   break;
+				case EVectorVMOp::enter_stat_scope:					InsPtr += 2;                    break;
+				case EVectorVMOp::exit_stat_scope:                                                  break;
+				case EVectorVMOp::update_id:
+				{
+					serializeIns(1, 2);
+					uint32 DataSetIdx     = VecIndices[0];
+					check(DataSetIdx < (uint32)VVMState->DataSets.Num());
+					FDataSetMeta *DataSet = &VVMState->DataSets[DataSetIdx];
+					int32 *R1             = (int32 *)VecReg[1];
+					int32 *R2             = (int32 *)VecReg[2];
+
+					check(DataSet->IDTable);
+					check(DataSet->IDTable->Num() >= DataSet->InstanceOffset + StartInstanceThisChunk + NumInstancesThisChunk);
+
+					int NumOutputInstances = BatchState->ChunkLocalData.NumOutputPerDataSet[DataSetIdx];
+					int NumFreed           = NumInstancesThisChunk - BatchState->ChunkLocalData.NumOutputPerDataSet[DataSetIdx];
+
+					//compute this chunk's MaxID
+					int MaxID = -1;
+					if (NumOutputInstances > 4)
+					{
+						int NumOutput4 = (int)(((((uint32)NumOutputInstances + 3U) & ~3U) - 1) >> 2);
+						VectorRegister4i Max4 = VectorIntSet1(-1);
+						check(((size_t)VecReg[1] & 0xF) == 0); //this must come from a register, we don't fuse input on update_id so we know it's aligned.  we can use pcmpgtd xmm, [ptr]
+						for (int i = 0; i < NumOutput4; ++i)
+						{
+							Max4 = VectorIntXor(VecReg[1][i].i, VectorIntAnd(VectorIntCompareGT(Max4, VecReg[1][i].i), VectorIntXor(Max4, VecReg[1][i].i)));
+						}
+						VectorRegister4i Last4 = VectorIntLoad(R1 + NumOutputInstances - 4);
+						Max4 = VectorIntXor(Last4, VectorIntAnd(VectorIntCompareGT(Max4, Last4), VectorIntXor(Max4, Last4)));
+						int M4[4];
+						VectorIntStore(Max4, M4);
+						int m0 = M4[0] > M4[1] ? M4[0] : M4[1];
+						int m1 = M4[2] > M4[3] ? M4[2] : M4[3];
+						int m = m0 > m1 ? m0 : m1;
+						if (m > MaxID)
+						{
+							MaxID = m;
+						}
+					}
+					else
+					{
+						for (int i = 0; i < NumOutputInstances; ++i)
+						{
+							if (R1[i] > MaxID)
+							{
+								MaxID = R1[i];
+							}
+						}
+					}
+					
+					// Update the actual index for this ID.  No thread safety is required as this ID slot can only ever be written by this instance
+					// The index passed into this function is the same as that given to the output* instructions
+					for (int i = 0; i < NumOutputInstances; ++i)
+					{
+						(*DataSet->IDTable)[R1[R2[i]]] = BatchState->ChunkLocalData.StartingOutputIdxPerDataSet[DataSetIdx] + i; //BatchState->ChunkLocalData.StartingOutputIdxPerDataSet[DataSetIdx] already has DataSet->InstanceOffset added to it
+					}
+
+					//Write the freed indices to the free table.
+					if (NumFreed > 0)
+					{
+						int StartNumFreed = FPlatformAtomics::InterlockedAdd((volatile int32 *)DataSet->NumFreeIDs, NumFreed);
+						int32 *FreeTableStart = DataSet->FreeIDTable->GetData() + StartNumFreed;
+						int c = 0;
+						int FreeCount = 0;
+						while (FreeCount < NumFreed)
+						{
+							check(c < NumInstancesThisChunk);
+							int d = R2[c] - c - FreeCount; //check for a gap in the write index and the counter... if nothing is freed then the write index matches the counter
+							if (d > 0)
+							{
+								VVMMemCpy(FreeTableStart + FreeCount, R1 + FreeCount + c, sizeof(int32) * d);
+								FreeCount += d;
+							}
+							++c;
+						}
+						check(FreeCount == NumFreed);
+					}
+
+					//Set the DataSet's MaxID if this chunk's MaxID is bigger
+					if (MaxID != -1)
+					{
+						int SanityCount = 0;
+						do {
+							int OldMaxID = *DataSet->MaxUsedID;
+							if (MaxID <= OldMaxID)
+							{
+								break;
+							}
+							int NewMaxID = FPlatformAtomics::InterlockedCompareExchange((volatile int32 *)DataSet->MaxUsedID, MaxID, OldMaxID);
+							if (NewMaxID == OldMaxID)
+							{
+								break;
+							}
+						} while (SanityCount++ < (1 << 30));
+						check(SanityCount < (1 << 30) - 1);
+					}
+					InsPtr += 6;
+				}
+				break;
+				case EVectorVMOp::acquire_id:
+				{
+					serializeIns(1, 2);
+					uint32 DataSetIdx = VecIndices[0];
+					check(DataSetIdx < (uint32)VVMState->DataSets.Num());
+					FDataSetMeta *DataSet = &VVMState->DataSets[DataSetIdx];
+					
+					{ //1. Get the free IDs into the temp register
+						int SanityCount = 0;
+						do
+						{
+							int OldNumFreeIDs = FPlatformAtomics::AtomicRead(DataSet->NumFreeIDs);
+							check(OldNumFreeIDs >= NumInstancesThisChunk);
+							//this is reverse-order from the original VM but it shouldn't matter since these are just re-used indices.
+							//VVMMemCpy(VecReg[1], DataSet->FreeIDTable->GetData() + OldNumFreeIDs - NumInstancesThisChunk, sizeof(int32) * NumInstancesThisChunk); //pull off the last added FreeIDs
+							int *OutPtr = (int *)VecReg[1];
+							int *InPtr  = DataSet->FreeIDTable->GetData() + OldNumFreeIDs - NumInstancesThisChunk;
+							for (int i = 0; i < NumInstancesThisChunk; ++i)
+							{
+								OutPtr[i] = InPtr[NumInstancesThisChunk - i - 1];
+							}
+							int NewNumFreeIDs = FPlatformAtomics::InterlockedCompareExchange((volatile int32 *)DataSet->NumFreeIDs, OldNumFreeIDs - NumInstancesThisChunk, OldNumFreeIDs);
+							if (NewNumFreeIDs == OldNumFreeIDs)
+							{
+								break;
+							}
+						} while (SanityCount++ < (1 << 30));
+						check(SanityCount < (1 << 30) - 1);
+					}
+					{ //2. append the IDs we acquired in step 1 to the end of the free table array, representing spawned IDs
+						//FreeID table is write-only as far as this invocation of the VM is concerned.
+						int StartNumSpawned = FPlatformAtomics::InterlockedAdd(DataSet->NumSpawnedIDs, NumInstancesThisChunk) + NumInstancesThisChunk;
+						check(StartNumSpawned <= DataSet->FreeIDTable->Max());
+						VVMMemCpy(DataSet->FreeIDTable->GetData() + DataSet->FreeIDTable->Max() - StartNumSpawned, VecReg[1], sizeof(int32) * NumInstancesThisChunk);
+					}
+					//3. set the tag
+					VVMMemSet32(VecReg[2], DataSet->IDAcquireTag, NumInstancesThisChunk);
+					InsPtr += 6;
+				}
+				break;
+				case EVectorVMOp::fused_input1_1: //op has 1 input operand and it's being overwritten to an input
+				{
+					OpCode                    = (EVectorVMOp)InsPtr[4];
+					uint8 RegType             = InsPtr[5];
+					uint16 DataSetIdx         = *(uint16 *)(InsPtr + 6);
+					uint16 InputRegIdx        = VecIndices[0];
+					uint32 InstanceOffset     = VVMState->DataSets[DataSetIdx].InstanceOffset;
+					uint32 InputRegTypeOffset = VVMState->DataSets[DataSetIdx].InputRegisterTypeOffsets[RegType];
+					uint32 **InputBuffers     = (uint32 **)VVMState->DataSets[DataSetIdx].InputRegisters.GetData();
+					uint32 *InputPtr          = InputBuffers[InputRegIdx + InputRegTypeOffset] + StartInstanceThisChunk + InstanceOffset;
+					VecReg[0]                 = (FVecReg *)InputPtr;
+					InsPtr += 4;
+					goto OpCodeSwitch;
+				}
+				break;
+				case EVectorVMOp::fused_input2_1: //op has 2 input operands, register 0 is being overwritten to an input (intentional fallthrough)
+				case EVectorVMOp::fused_input2_2: //op has 2 input operands, register 1 is being overwritten to an input
+				{
+					int RegToSwitchToInput     = (int)OpCode - (int)EVectorVMOp::fused_input2_1;
+					OpCode                     = (EVectorVMOp)InsPtr[6];
+					uint8 RegType              = InsPtr[7];
+					uint16 DataSetIdx          = *(uint16 *)(InsPtr + 8);
+					uint16 InputRegIdx         = VecIndices[RegToSwitchToInput];
+					uint32 InstanceOffset      = VVMState->DataSets[DataSetIdx].InstanceOffset;
+					uint32 InputRegTypeOffset  = VVMState->DataSets[DataSetIdx].InputRegisterTypeOffsets[RegType];
+					uint32 **InputBuffers      = (uint32 **)VVMState->DataSets[DataSetIdx].InputRegisters.GetData();
+					uint32 *InputPtr           = InputBuffers[InputRegIdx + InputRegTypeOffset] + StartInstanceThisChunk + InstanceOffset;
+					VecReg[RegToSwitchToInput] = (FVecReg *)InputPtr;
+					InsPtr += 4;
+					goto OpCodeSwitch;
+				}
+				break;
+				case EVectorVMOp::fused_input2_3: //op has 2 input operands, register 0 and 1 are being overwritten to an input
+				{
+					OpCode                       = (EVectorVMOp)InsPtr[6];
+					uint8 RegType[2]             = { InsPtr[7], InsPtr[10] };
+					uint16 DataSetIdx[2]         = { *(uint16 *)(InsPtr + 8), *(uint16 *)(InsPtr + 11) };
+					int32 InstanceOffset[2]      = { VVMState->DataSets[DataSetIdx[0]].InstanceOffset, VVMState->DataSets[DataSetIdx[1]].InstanceOffset };
+					uint32 InputRegTypeOffset[2] = { VVMState->DataSets[DataSetIdx[0]].InputRegisterTypeOffsets[RegType[0]], VVMState->DataSets[DataSetIdx[1]].InputRegisterTypeOffsets[RegType[1]] };
+					uint32 **InputBuffers[2]     = { (uint32 **)VVMState->DataSets[DataSetIdx[0]].InputRegisters.GetData(), (uint32 **)VVMState->DataSets[DataSetIdx[1]].InputRegisters.GetData() };
+					uint32 *InputPtr[2]          = { InputBuffers[0][VecIndices[0] + InputRegTypeOffset[0]] + StartInstanceThisChunk + InstanceOffset[0], InputBuffers[1][VecIndices[1] + InputRegTypeOffset[1]] + StartInstanceThisChunk + InstanceOffset[1] };
+					VecReg[0]                    = (FVecReg *)InputPtr[0];
+					VecReg[1]                    = (FVecReg *)InputPtr[1];
+					InsPtr += 7;
+					goto OpCodeSwitch;
+				}
+				break;
+				case EVectorVMOp::fused_input3_1: //op has 3 input operands, register 0 is being overwritten to an input (intentional fallthrough)
+				case EVectorVMOp::fused_input3_2: //op has 3 input operands, register 1 is being overwritten to an input (intentional fallthrough)
+				case EVectorVMOp::fused_input3_4: //op has 3 input operands, register 2 is being overwritten to an input
+				{
+					int RegToSwitchToInput     = (int)OpCode - (int)EVectorVMOp::fused_input3_1;
+					check(RegToSwitchToInput >= 0 && RegToSwitchToInput <= 2);
+					OpCode                     = (EVectorVMOp)InsPtr[8];
+					uint8 RegType              = InsPtr[9];
+					uint16 DataSetIdx          = *(uint16 *)(InsPtr + 10);
+					uint16 InputRegIdx         = VecIndices[RegToSwitchToInput];
+					uint32 InstanceOffset      = VVMState->DataSets[DataSetIdx].InstanceOffset;
+					uint32 InputRegTypeOffset  = VVMState->DataSets[DataSetIdx].InputRegisterTypeOffsets[RegType];
+					uint32 **InputBuffers      = (uint32 **)VVMState->DataSets[DataSetIdx].InputRegisters.GetData();
+					uint32 *InputPtr           = InputBuffers[InputRegIdx + InputRegTypeOffset] + StartInstanceThisChunk + InstanceOffset;
+					VecReg[RegToSwitchToInput] = (FVecReg *)InputPtr;
+					InsPtr += 4;
+					goto OpCodeSwitch;
+				}
+				break;
+				case EVectorVMOp::fused_input3_3: //op has 3 input operands, register 0 and 1 are being overwritten to inputs (intentional fallthrough)
+				case EVectorVMOp::fused_input3_5: //op has 3 input operands, register 0 and 2 are being overwritten to inputs
+				{
+					int RegIdx2                  = 1 + (int)OpCode - (int)EVectorVMOp::fused_input3_3;
+					OpCode                       = (EVectorVMOp)InsPtr[8];
+					uint8 RegType[2]             = { InsPtr[9], InsPtr[12] };
+					uint16 DataSetIdx[2]         = { *(uint16 *)(InsPtr + 10), *(uint16 *)(InsPtr + 13) };
+					int32 InstanceOffset[2]      = { VVMState->DataSets[DataSetIdx[0]].InstanceOffset, VVMState->DataSets[DataSetIdx[1]].InstanceOffset };
+					uint32 InputRegTypeOffset[2] = { VVMState->DataSets[DataSetIdx[0]].InputRegisterTypeOffsets[RegType[0]], VVMState->DataSets[DataSetIdx[1]].InputRegisterTypeOffsets[RegType[1]] };
+					uint32 **InputBuffers[2]     = { (uint32 **)VVMState->DataSets[DataSetIdx[0]].InputRegisters.GetData(), (uint32 **)VVMState->DataSets[DataSetIdx[1]].InputRegisters.GetData() };
+					uint32 *InputPtr[2]          = { InputBuffers[0][VecIndices[0] + InputRegTypeOffset[0]] + StartInstanceThisChunk + InstanceOffset[0], InputBuffers[1][VecIndices[RegIdx2] + InputRegTypeOffset[1]] + StartInstanceThisChunk + InstanceOffset[1] };
+					VecReg[0]       = (FVecReg *)InputPtr[0];
+					VecReg[RegIdx2] = (FVecReg *)InputPtr[1];
+					InsPtr += 7;
+					goto OpCodeSwitch;
+				}
+				break;
+				case EVectorVMOp::fused_input3_6: //op has 3 input operands, register 1 and 2 are being overwritten to inputs
+				{
+					OpCode                       = (EVectorVMOp)InsPtr[8];
+					uint8 RegType[2]             = { InsPtr[9], InsPtr[12] };
+					uint16 DataSetIdx[2]         = { *(uint16 *)(InsPtr + 10), *(uint16 *)(InsPtr + 13) };
+					int32 InstanceOffset[2]      = { VVMState->DataSets[DataSetIdx[0]].InstanceOffset, VVMState->DataSets[DataSetIdx[1]].InstanceOffset };
+					uint32 InputRegTypeOffset[2] = { VVMState->DataSets[DataSetIdx[0]].InputRegisterTypeOffsets[RegType[0]], VVMState->DataSets[DataSetIdx[1]].InputRegisterTypeOffsets[RegType[1]] };
+					uint32 **InputBuffers[2]     = { (uint32 **)VVMState->DataSets[DataSetIdx[0]].InputRegisters.GetData(), (uint32 **)VVMState->DataSets[DataSetIdx[1]].InputRegisters.GetData() };
+					uint32 *InputPtr[2]          = { InputBuffers[0][VecIndices[1] + InputRegTypeOffset[0]] + StartInstanceThisChunk + InstanceOffset[0], InputBuffers[1][VecIndices[2] + InputRegTypeOffset[1]] + StartInstanceThisChunk + InstanceOffset[1] };
+					VecReg[1] = (FVecReg *)InputPtr[0];
+					VecReg[2] = (FVecReg *)InputPtr[1];
+					InsPtr += 7;
+					goto OpCodeSwitch;
+				}
+				break;
+				case EVectorVMOp::fused_input3_7: //op has 3 input operands, register 1, 2, and 3 are all being overwritten to inputs
+				{
+					OpCode                       = (EVectorVMOp)InsPtr[8];
+					uint8 RegType[3]             = { InsPtr[9], InsPtr[12], InsPtr[15] };
+					uint16 DataSetIdx[3]         = { *(uint16 *)(InsPtr + 10), *(uint16 *)(InsPtr + 13), *(uint16 *)(InsPtr + 16) };
+					int32 InstanceOffset[3]      = { VVMState->DataSets[DataSetIdx[0]].InstanceOffset, 
+						                             VVMState->DataSets[DataSetIdx[1]].InstanceOffset, 
+						                             VVMState->DataSets[DataSetIdx[2]].InstanceOffset
+					                               };
+					uint32 InputRegTypeOffset[3] = { VVMState->DataSets[DataSetIdx[0]].InputRegisterTypeOffsets[RegType[0]], 
+						                             VVMState->DataSets[DataSetIdx[1]].InputRegisterTypeOffsets[RegType[1]], 
+						                             VVMState->DataSets[DataSetIdx[2]].InputRegisterTypeOffsets[RegType[2]] 
+					                               };
+					uint32 **InputBuffers[3]     = { (uint32 **)VVMState->DataSets[DataSetIdx[0]].InputRegisters.GetData(), 
+						                             (uint32 **)VVMState->DataSets[DataSetIdx[1]].InputRegisters.GetData(), 
+						                             (uint32 **)VVMState->DataSets[DataSetIdx[2]].InputRegisters.GetData() 
+					                               };
+					uint32 *InputPtr[3]          = { InputBuffers[0][VecIndices[0] + InputRegTypeOffset[0]] + StartInstanceThisChunk + InstanceOffset[0], 
+						                             InputBuffers[1][VecIndices[1] + InputRegTypeOffset[1]] + StartInstanceThisChunk + InstanceOffset[1], 
+						                             InputBuffers[2][VecIndices[2] + InputRegTypeOffset[2]] + StartInstanceThisChunk + InstanceOffset[2]
+					                               };
+					VecReg[0] = (FVecReg *)InputPtr[0];
+					VecReg[1] = (FVecReg *)InputPtr[1];
+					VecReg[2] = (FVecReg *)InputPtr[2];
+					InsPtr += 10;
+					goto OpCodeSwitch;
+				}
+				break;
+				case EVectorVMOp::copy_to_output:
+				{
+					uint16 OutputDataSetIdx   = VecIndices[0];
+					uint16 InputDataSetIdx    = VecIndices[1];
+					uint16 OutputDstIdxRegIdx = VecIndices[2];
+					uint8 RegType             = InsPtr[6];
+					uint8 Count               = InsPtr[7];
+					InsPtr += 8;
+
+					//Output
+					uint32 NumOutputInstances   = BatchState->ChunkLocalData.NumOutputPerDataSet[OutputDataSetIdx];
+					uint32 OutputInstanceOffset = BatchState->ChunkLocalData.StartingOutputIdxPerDataSet[OutputDataSetIdx];
+					uint32 OutputRegTypeOffset  = VVMState->DataSets[OutputDataSetIdx].OutputRegisterTypeOffsets[RegType];
+					uint32 **OutputBuffers      = (uint32 **)VVMState->DataSets[OutputDataSetIdx].OutputRegisters.GetData();
+
+					//Input
+					uint32 InputInstanceOffset = VVMState->DataSets[InputDataSetIdx].InstanceOffset;
+					uint32 InputRegTypeOffset  = VVMState->DataSets[InputDataSetIdx].InputRegisterTypeOffsets[RegType];
+					uint32 **InputBuffers      = (uint32 **)VVMState->DataSets[InputDataSetIdx].InputRegisters.GetData();
+
+					int *DstIdxReg             = (int *)VecReg[2];
+					int StartSrcIndex          = StartInstanceThisChunk + InputInstanceOffset;
+					uint16 *IdxInsPtr          = (uint16 *)InsPtr;
+					if (NumOutputInstances == BatchState->NumInstances)
+					{
+						//if we're writing the same number of inputs as outputs then we can just memcpy
+						for (int i = 0; i < (int)Count; ++i)
+						{
+							uint16 OutputDstIdx = IdxInsPtr[(i << 1) + 0];
+							uint16 InputSrcIdx  = IdxInsPtr[(i << 1) + 1];
+							uint32 *SrcBuffer   = InputBuffers [InputSrcIdx  + InputRegTypeOffset ] + StartSrcIndex;
+							uint32 *DstBuffer   = OutputBuffers[OutputDstIdx + OutputRegTypeOffset] + OutputInstanceOffset;
+							VVMMemCpy(DstBuffer, SrcBuffer, sizeof(uint32) * NumOutputInstances);
+						}
+					}
+					else
+					{
+						//if we are discarding at least one instance we can't just copy the memory and need to use the index generated in acquire_index
+						for (int i = 0; i < (int)Count; ++i)
+						{
+							uint16 OutputDstIdx = IdxInsPtr[(i << 1) + 0];
+							uint16 InputSrcIdx  = IdxInsPtr[(i << 1) + 1];
+							uint32 *SrcBuffer   = InputBuffers [InputSrcIdx  + InputRegTypeOffset ] + StartSrcIndex;
+							uint32 *DstBuffer   = OutputBuffers[OutputDstIdx + OutputRegTypeOffset] + OutputInstanceOffset;
+							for (uint32 j = 0; j < NumOutputInstances; ++j)
+							{
+								DstBuffer[j] = SrcBuffer[DstIdxReg[j]];
+							}
+						}
+					}
+					InsPtr += (int)Count * 4;
+				}
+				break;
+				case EVectorVMOp::output_batch2:
+				{
+					int *DstIdxReg            = (int *)VecReg[0];
+					uint16 DataSetIdx         = VecIndices[3];
+					uint32 RegTypeOffset      = VVMState->DataSets[DataSetIdx].OutputRegisterTypeOffsets[InsPtr[12]];
+					uint32 **OutputBuffers    = (uint32 **)VVMState->DataSets[DataSetIdx].OutputRegisters.GetData();
+					uint32 InstanceOffset     = BatchState->ChunkLocalData.StartingOutputIdxPerDataSet[DataSetIdx];
+					uint32 NumOutputInstances = BatchState->ChunkLocalData.NumOutputPerDataSet[DataSetIdx];
+
+					uint32 *DstReg[2]         = { OutputBuffers[RegTypeOffset + VecIndices[4]] + InstanceOffset, 
+						                          OutputBuffers[RegTypeOffset + VecIndices[5]] + InstanceOffset
+												};
+					uint32 *SrcReg[2]         = { (uint32 *)VecReg[1], 
+					                              (uint32 *)VecReg[2] };
+					
+					if (NumOutputInstances == BatchState->NumInstances)
+					{
+						if (RegInc[1] == 0)
+						{
+							VVMMemSet32(DstReg[0], *SrcReg[0], NumOutputInstances);
+						}
+						else
+						{
+							VVMMemCpy(DstReg[0], SrcReg[0], sizeof(uint32) * NumOutputInstances);
+						}
+						if (RegInc[2] == 0)
+						{
+							VVMMemSet32(DstReg[1], *SrcReg[1], NumOutputInstances);
+						}
+						else
+						{
+							VVMMemCpy(DstReg[1], SrcReg[1], sizeof(uint32) * NumOutputInstances);
+						}
+					}
+					else
+					{
+						for (uint32 i = 0; i < NumOutputInstances; ++i)
+						{
+							DstReg[0][i] = SrcReg[0][DstIdxReg[i & RegInc[0]] & RegInc[1]];
+							DstReg[1][i] = SrcReg[1][DstIdxReg[i & RegInc[0]] & RegInc[2]];
+						}
+					}
+					InsPtr += 13;
+				}
+				break;
+				case EVectorVMOp::output_batch3:
+				{
+					int *DstIdxReg            = (int *)VecReg[0];
+					uint16 DataSetIdx         = VecIndices[4];
+					uint32 RegTypeOffset      = VVMState->DataSets[DataSetIdx].OutputRegisterTypeOffsets[InsPtr[16]];
+					uint32 **OutputBuffers    = (uint32 **)VVMState->DataSets[DataSetIdx].OutputRegisters.GetData();
+					uint32 InstanceOffset     = BatchState->ChunkLocalData.StartingOutputIdxPerDataSet[DataSetIdx];
+					uint32 NumOutputInstances = BatchState->ChunkLocalData.NumOutputPerDataSet[DataSetIdx];
+
+					uint32 *DstReg[3]         = { OutputBuffers[RegTypeOffset + VecIndices[5]] + InstanceOffset, 
+						                          OutputBuffers[RegTypeOffset + VecIndices[6]] + InstanceOffset,
+					                              OutputBuffers[RegTypeOffset + VecIndices[7]] + InstanceOffset };
+					uint32 *SrcReg[3]         = { (uint32 *)VecReg[1], 
+					                              (uint32 *)VecReg[2], 
+					                              (uint32 *)VecReg[3] };
+
+					if (NumOutputInstances == BatchState->NumInstances)
+					{
+						if (RegInc[1] == 0)
+						{
+							VVMMemSet32(DstReg[0], *SrcReg[0], NumOutputInstances);
+						}
+						else
+						{
+							VVMMemCpy(DstReg[0], SrcReg[0], sizeof(uint32) * NumOutputInstances);
+						}
+						if (RegInc[2] == 0)
+						{
+							VVMMemSet32(DstReg[1], *SrcReg[1], NumOutputInstances);
+						}
+						else
+						{
+							VVMMemCpy(DstReg[1], SrcReg[1], sizeof(uint32) * NumOutputInstances);
+						}
+						if (RegInc[3] == 0)
+						{
+							VVMMemSet32(DstReg[2], *SrcReg[2], NumOutputInstances);
+						}
+						else
+						{
+							VVMMemCpy(DstReg[2], SrcReg[2], sizeof(uint32) * NumOutputInstances);
+						}
+					}
+					else
+					{
+						for (uint32 i = 0; i < NumOutputInstances; ++i)
+						{
+							DstReg[0][i] = SrcReg[0][DstIdxReg[i & RegInc[0]] & RegInc[1]];
+							DstReg[1][i] = SrcReg[1][DstIdxReg[i & RegInc[0]] & RegInc[2]];
+							DstReg[2][i] = SrcReg[2][DstIdxReg[i & RegInc[0]] & RegInc[3]];
+						}
+					}
+					InsPtr += 17;
+				}
+				break;
+				case EVectorVMOp::output_batch4:
+				{
+					uint16 *OutputIndices     = VecIndices + 4;
+					uint16 DataSetIdx         = VecIndices[8];
+					uint16 IdxRegIdx          = VecIndices[9];
+					int *DstIdxReg            = (int *)(BatchState->RegisterData + VVMState->NumConstBuffers + NumLoops * IdxRegIdx); //guaranteed by the optimizer to be a temp register and not a const
+					uint32 RegTypeOffset      = VVMState->DataSets[DataSetIdx].OutputRegisterTypeOffsets[InsPtr[20]];
+					uint32 **OutputBuffers    = (uint32 **)VVMState->DataSets[DataSetIdx].OutputRegisters.GetData();
+					uint32 InstanceOffset     = BatchState->ChunkLocalData.StartingOutputIdxPerDataSet[DataSetIdx];
+					uint32 NumOutputInstances = BatchState->ChunkLocalData.NumOutputPerDataSet[DataSetIdx];
+
+					uint32 *DstReg[4]         = { OutputBuffers[RegTypeOffset + OutputIndices[0]] + InstanceOffset, 
+						                          OutputBuffers[RegTypeOffset + OutputIndices[1]] + InstanceOffset, 
+						                          OutputBuffers[RegTypeOffset + OutputIndices[2]] + InstanceOffset, 
+						                          OutputBuffers[RegTypeOffset + OutputIndices[3]] + InstanceOffset };
+
+					uint32 *SrcReg[4]         = { (uint32 *)VecReg[0], 
+						                          (uint32 *)VecReg[1], 
+						                          (uint32 *)VecReg[2], 
+						                          (uint32 *)VecReg[3] };
+					
+					if (NumOutputInstances == BatchState->NumInstances)
+					{
+						if (RegInc[0] == 0)
+						{
+							VVMMemSet32(DstReg[0], *SrcReg[0], NumOutputInstances);
+						}
+						else
+						{
+							VVMMemCpy(DstReg[0], SrcReg[0], sizeof(uint32) * NumOutputInstances);
+						}
+						if (RegInc[1] == 0)
+						{
+							VVMMemSet32(DstReg[1], *SrcReg[1], NumOutputInstances);
+						}
+						else
+						{
+							VVMMemCpy(DstReg[1], SrcReg[1], sizeof(uint32) * NumOutputInstances);
+						}
+						if (RegInc[2] == 0)
+						{
+							VVMMemSet32(DstReg[2], *SrcReg[2], NumOutputInstances);
+						}
+						else
+						{
+							VVMMemCpy(DstReg[2], SrcReg[2], sizeof(uint32) * NumOutputInstances);
+						}
+						if (RegInc[3] == 0)
+						{
+							VVMMemSet32(DstReg[3], *SrcReg[3], NumOutputInstances);
+						}
+						else
+						{
+							VVMMemCpy(DstReg[3], SrcReg[3], sizeof(uint32) * NumOutputInstances);
+						}
+					}
+					else
+					{
+						for (uint32 i = 0; i < NumOutputInstances; ++i)
+						{
+							DstReg[0][i] = SrcReg[0][DstIdxReg[i] & RegInc[0]];
+							DstReg[1][i] = SrcReg[1][DstIdxReg[i] & RegInc[1]];
+							DstReg[2][i] = SrcReg[2][DstIdxReg[i] & RegInc[2]];
+							DstReg[3][i] = SrcReg[3][DstIdxReg[i] & RegInc[3]];
+						}
+					}
+					InsPtr += 21;
+				}
+				break;
+				case EVectorVMOp::output_batch7:
+				{
+					uint16 DataSetIdx         = VecIndices[4];
+					int *DstIdxReg            = (int *)VecReg[0];
+					uint32 RegTypeOffset      = VVMState->DataSets[DataSetIdx].OutputRegisterTypeOffsets[InsPtr[32]];
+					uint32 **OutputBuffers    = (uint32 **)VVMState->DataSets[DataSetIdx].OutputRegisters.GetData();
+					uint32 InstanceOffset     = BatchState->ChunkLocalData.StartingOutputIdxPerDataSet[DataSetIdx];
+					uint32 NumOutputInstances = BatchState->ChunkLocalData.NumOutputPerDataSet[DataSetIdx];
+
+					{ //first 3
+						uint32 *DstReg[3]         = { OutputBuffers[RegTypeOffset + VecIndices[9] ] + InstanceOffset, 
+													  OutputBuffers[RegTypeOffset + VecIndices[10]] + InstanceOffset,
+													  OutputBuffers[RegTypeOffset + VecIndices[11]] + InstanceOffset };
+						uint32 *SrcReg[3]         = { (uint32 *)VecReg[1], 
+													  (uint32 *)VecReg[2], 
+													  (uint32 *)VecReg[3] };
+
+					
+						if (NumOutputInstances == BatchState->NumInstances)
+						{
+							if (RegInc[1] == 0)
+							{
+								VVMMemSet32(DstReg[0], *SrcReg[0], NumOutputInstances);
+							}
+							else
+							{
+								VVMMemCpy(DstReg[0], SrcReg[0], sizeof(uint32) * NumOutputInstances);
+							}
+							if (RegInc[2] == 0)
+							{
+								VVMMemSet32(DstReg[1], *SrcReg[1], NumOutputInstances);
+							} 
+							else
+							{
+								VVMMemCpy(DstReg[1], SrcReg[1], sizeof(uint32) * NumOutputInstances);
+							}
+							if (RegInc[3] == 0)
+							{
+								VVMMemSet32(DstReg[2], *SrcReg[2], NumOutputInstances);
+							}
+							else
+							{
+								VVMMemCpy(DstReg[2], SrcReg[2], sizeof(uint32) * NumOutputInstances);
+							}
+						}
+						else
+						{
+							for (uint32 i = 0; i < NumOutputInstances; ++i)
+							{
+								DstReg[0][i] = SrcReg[0][DstIdxReg[i & RegInc[0]] & RegInc[1]];
+								DstReg[1][i] = SrcReg[1][DstIdxReg[i & RegInc[0]] & RegInc[2]];
+								DstReg[2][i] = SrcReg[2][DstIdxReg[i & RegInc[0]] & RegInc[3]];
+							}
+						}
+					}
+					VVMDecodeInstructionRegisters(InsPtr + 10, VecReg, RegInc);
+					{ //next 4
+						uint32 *DstReg[4]         = { OutputBuffers[RegTypeOffset + VecIndices[12]] + InstanceOffset, 
+													  OutputBuffers[RegTypeOffset + VecIndices[13]] + InstanceOffset, 
+													  OutputBuffers[RegTypeOffset + VecIndices[14]] + InstanceOffset, 
+													  OutputBuffers[RegTypeOffset + VecIndices[15]] + InstanceOffset };
+
+						uint32 *SrcReg[4]         = { (uint32 *)VecReg[0], 
+													  (uint32 *)VecReg[1], 
+													  (uint32 *)VecReg[2], 
+													  (uint32 *)VecReg[3] };
+						if (NumOutputInstances == BatchState->NumInstances)
+						{
+							if (RegInc[0] == 0)
+							{
+								VVMMemSet32(DstReg[0], *SrcReg[0], NumOutputInstances);
+							}
+							else
+							{
+								VVMMemCpy(DstReg[0], SrcReg[0], sizeof(uint32) * NumOutputInstances);
+							}
+							if (RegInc[1] == 0)
+							{
+								VVMMemSet32(DstReg[1], *SrcReg[1], NumOutputInstances);
+							}
+							else
+							{
+								VVMMemCpy(DstReg[1], SrcReg[1], sizeof(uint32) * NumOutputInstances);
+							}
+							if (RegInc[2] == 0)
+							{
+								VVMMemSet32(DstReg[2], *SrcReg[2], NumOutputInstances);
+							}
+							else
+							{
+								VVMMemCpy(DstReg[2], SrcReg[2], sizeof(uint32) * NumOutputInstances);
+							}
+							if (RegInc[3] == 0)
+							{
+								VVMMemSet32(DstReg[3], *SrcReg[3], NumOutputInstances);
+							}
+							else
+							{
+								VVMMemCpy(DstReg[3], SrcReg[3], sizeof(uint32) * NumOutputInstances);
+							}
+						}
+						else
+						{
+							for (uint32 i = 0; i < NumOutputInstances; ++i)
+							{
+								DstReg[0][i] = SrcReg[0][DstIdxReg[i] & RegInc[0]];
+								DstReg[1][i] = SrcReg[1][DstIdxReg[i] & RegInc[1]];
+								DstReg[2][i] = SrcReg[2][DstIdxReg[i] & RegInc[2]];
+								DstReg[3][i] = SrcReg[3][DstIdxReg[i] & RegInc[3]];
+							}
+						}
+					}
+					InsPtr += 33;
+				}
+				break;
+				case EVectorVMOp::output_batch8:
+				{
+					uint16 *OutputIndices     = VecIndices + 8;
+					uint16 DataSetIdx         = VecIndices[16];
+					uint16 IdxRegIdx          = VecIndices[17];
+					int *DstIdxReg            = (int *)(BatchState->RegisterData + VVMState->NumConstBuffers + NumLoops * IdxRegIdx); //guaranteed by the optimizer to be a temp register and not a const
+					uint32 RegTypeOffset      = VVMState->DataSets[DataSetIdx].OutputRegisterTypeOffsets[InsPtr[36]];
+					uint32 **OutputBuffers    = (uint32 **)VVMState->DataSets[DataSetIdx].OutputRegisters.GetData();
+					uint32 InstanceOffset     = BatchState->ChunkLocalData.StartingOutputIdxPerDataSet[DataSetIdx];
+					uint32 NumOutputInstances = BatchState->ChunkLocalData.NumOutputPerDataSet[DataSetIdx];
+					
+					{ //first block of 4
+						uint32 *DstReg[4]         = { OutputBuffers[RegTypeOffset + OutputIndices[0]] + InstanceOffset, 
+													  OutputBuffers[RegTypeOffset + OutputIndices[1]] + InstanceOffset, 
+													  OutputBuffers[RegTypeOffset + OutputIndices[2]] + InstanceOffset, 
+													  OutputBuffers[RegTypeOffset + OutputIndices[3]] + InstanceOffset };
+
+						uint32 *SrcReg[4]         = { (uint32 *)VecReg[0], 
+													  (uint32 *)VecReg[1], 
+													  (uint32 *)VecReg[2], 
+													  (uint32 *)VecReg[3] };
+
+						if (NumOutputInstances == BatchState->NumInstances)
+						{
+							if (RegInc[0] == 0)
+							{
+								VVMMemSet32(DstReg[0], *SrcReg[0], NumOutputInstances);
+							}
+							else
+							{
+								VVMMemCpy(DstReg[0], SrcReg[0], sizeof(uint32) * NumOutputInstances);
+							}
+							if (RegInc[1] == 0)
+							{
+								VVMMemSet32(DstReg[1], *SrcReg[1], NumOutputInstances);
+							}
+							else
+							{
+								VVMMemCpy(DstReg[1], SrcReg[1], sizeof(uint32) * NumOutputInstances);
+							}
+							if (RegInc[2] == 0)
+							{
+								VVMMemSet32(DstReg[2], *SrcReg[2], NumOutputInstances);
+							}
+							else
+							{
+								VVMMemCpy(DstReg[2], SrcReg[2], sizeof(uint32) * NumOutputInstances);
+							}
+							if (RegInc[3] == 0)
+							{
+								VVMMemSet32(DstReg[3], *SrcReg[3], NumOutputInstances);
+							}
+							else
+							{
+								VVMMemCpy(DstReg[3], SrcReg[3], sizeof(uint32) * NumOutputInstances);
+							}
+						}
+						else
+						{
+							for (uint32 i = 0; i < NumOutputInstances; ++i)
+							{
+								DstReg[0][i] = SrcReg[0][DstIdxReg[i] & RegInc[0]];
+								DstReg[1][i] = SrcReg[1][DstIdxReg[i] & RegInc[1]];
+								DstReg[2][i] = SrcReg[2][DstIdxReg[i] & RegInc[2]];
+								DstReg[3][i] = SrcReg[3][DstIdxReg[i] & RegInc[3]];
+							}
+						}
+					}
+					VVMDecodeInstructionRegisters(InsPtr + 8, VecReg, RegInc);
+					{ //second block of 4
+						uint32 *DstReg[4]         = { OutputBuffers[RegTypeOffset + OutputIndices[4]] + InstanceOffset, 
+													  OutputBuffers[RegTypeOffset + OutputIndices[5]] + InstanceOffset, 
+													  OutputBuffers[RegTypeOffset + OutputIndices[6]] + InstanceOffset, 
+													  OutputBuffers[RegTypeOffset + OutputIndices[7]] + InstanceOffset };
+
+						uint32 *SrcReg[4]         = { (uint32 *)VecReg[0], 
+													  (uint32 *)VecReg[1], 
+													  (uint32 *)VecReg[2], 
+													  (uint32 *)VecReg[3] };
+						if (NumOutputInstances == BatchState->NumInstances)
+						{
+							if (RegInc[0] == 0)
+							{
+								VVMMemSet32(DstReg[0], *SrcReg[0], NumOutputInstances);
+							}
+							else
+							{
+								VVMMemCpy(DstReg[0], SrcReg[0], sizeof(uint32) * NumOutputInstances);
+							}
+							if (RegInc[1] == 0)
+							{
+								VVMMemSet32(DstReg[1], *SrcReg[1], NumOutputInstances);
+							}
+							else
+							{
+								VVMMemCpy(DstReg[1], SrcReg[1], sizeof(uint32) * NumOutputInstances);
+							}
+							if (RegInc[2] == 0)
+							{
+								VVMMemSet32(DstReg[2], *SrcReg[2], NumOutputInstances);
+							}
+							else
+							{
+								VVMMemCpy(DstReg[2], SrcReg[2], sizeof(uint32) * NumOutputInstances);
+							}
+							if (RegInc[3] == 0)
+							{
+								VVMMemSet32(DstReg[3], *SrcReg[3], NumOutputInstances);
+							}
+							else
+							{
+								VVMMemCpy(DstReg[3], SrcReg[3], sizeof(uint32) * NumOutputInstances);
+							}
+						}
+						else
+						{
+							for (uint32 i = 0; i < NumOutputInstances; ++i)
+							{
+								DstReg[0][i] = SrcReg[0][DstIdxReg[i] & RegInc[0]];
+								DstReg[1][i] = SrcReg[1][DstIdxReg[i] & RegInc[1]];
+								DstReg[2][i] = SrcReg[2][DstIdxReg[i] & RegInc[2]];
+								DstReg[3][i] = SrcReg[3][DstIdxReg[i] & RegInc[3]];
+							}
+						}
+					}
+					InsPtr += 37;
+				} break;
+				default: break;
+			}
+			VVMSer_insEndExp(SerializeState, (int)(VVMSerStartOpPtr - VVMState->Bytecode), (int)(InsPtr - VVMSerStartOpPtr));
+		}
+		VVMSer_chunkEndExp(SerializeState);
+	}
+	VVMSer_batchEndExp(SerializeState);
+
+	int32 WasCurrentlyExecuting = FPlatformAtomics::InterlockedCompareExchange(&BatchState->CurrentlyExecuting, 0, 1);
+	check(WasCurrentlyExecuting == 1); //Sanity test to make sure that the CurrentlyExcuting flag was never changed during execution of this batch
+
+	FPlatformAtomics::InterlockedAdd(&VVMState->NumInstancesCompleted, BatchState->NumInstances);
+}
+
+VECTORVM_API void ExecVectorVMState(FVectorVMState *VVMState, FVectorVMSerializeState *SerializeState, FVectorVMSerializeState *CmpSerializeState)
+{
+#if defined(VVM_INCLUDE_SERIALIZATION) && !defined(VVM_SERIALIZE_NO_WRITE)
+	uint64 StartTime = FPlatformTime::Cycles64();
+	if (SerializeState)
+	{
+		SerializeState->ExecDt = 0;
+		SerializeState->SerializeDt = 0;
+	}
+#endif //VVM_INCLUDE_SERIALIZATION
+
+	if (VVMState->NumBatches > 1)
+	{
+#ifdef VVM_USE_OFFLINE_THREADING
+		if (parallelJobFn)
+		{
+			for (int i = 0; i < VVMState->NumBatches; ++i)
+			{
+				parallelJobFn(ExecVVMBatch, VVMState, i, SerializeState, CmpSerializeState);
+			}
+		}
+		else
+		{
+			for (int i = 0; i < VVMState->NumBatches; ++i)
+			{
+				ExecVVMBatch(VVMState, i, SerializeState, CmpSerializeState);
+			}
+		}
+#else
+		auto ExecChunkBatch = [&](int32 BatchIdx)
+		{
+			ExecVVMBatch(VVMState, BatchIdx, SerializeState, CmpSerializeState);
+		};
+		ParallelFor(VVMState->NumBatches, ExecChunkBatch, true);// GbParallelVVM == 0 || !bParallel);
+#endif
+	}
+	else
+	{
+		ExecVVMBatch(VVMState, 0, SerializeState, CmpSerializeState);
+	}
+
+#ifdef VVM_USE_OFFLINE_THREADING
+	//Unreal's ParallelFor() will block the executing thread until it's finished.  That isn't guaranteed
+	//outside of UE, (ie: the debugger, the only other thing that uses this as of this writing) so block.
+	while (VVMState->NumInstancesCompleted < VVMState->TotalNumInstances)
+	{
+		FPlatformProcess::Yield();
+	}
+#endif
+
+	for (uint32 i = 0; i < VVMState->NumOutputDataSets; ++i)
+	{
+		VVMState->DataSets[i].DataSetAccessIndex = VVMState->NumOutputPerDataSet[i] - 1;
+	}
+
+#if defined(VVM_INCLUDE_SERIALIZATION) && !defined(VVM_SERIALIZE_NO_WRITE)
+	uint64 EndTime = FPlatformTime::Cycles64();
+	if (SerializeState)
+	{
+		SerializeState->ExecDt = EndTime - StartTime;
+	}
+#endif //VVM_INCLUDE_SERIALIZATION
+}
+
+#undef VVM_MIN
+#undef VVM_MAX
+#undef VVM_CLAMP
+#undef VVM_ALIGN
+#undef VVM_ALIGN_4
+#undef VVM_ALIGN_16
+#undef VVM_ALIGN_64
+
+#undef VVMSet_m128Const
+#undef VVMSet_m128iConst
+#undef VVMSet_m128iConst4
+#undef VVM_m128Const
+#undef VVM_m128iConst
+
+//PRAGMA_ENABLE_OPTIMIZATION
+
+#else //NIAGARA_EXP_VM
+
+VECTORVM_API FVectorVMState *InitVectorVMState(struct FVectorVMInitData *InitData, struct FVectorVMExternalFnPerInstanceData **OutPerInstanceExtData, struct FVectorVMSerializeState *SerializeState)
+{
+	return nullptr;
+}
+
+VECTORVM_API void FreeVectorVMState(FVectorVMState *VVMState)
+{
+
+}
+
+VECTORVM_API void ExecVectorVMState(FVectorVMState *VVMState, FVectorVMSerializeState *SerializeState, FVectorVMSerializeState *CmpSerializeState)
+{
+
+}
+
+#endif //NIAGARA_EXP_VM
