- New Vector VM with optimizer and serialization.

diff --git a/Engine/Source/Runtime/VectorVM/Private/VectorVMExperimental_Optimizer.inl b/Engine/Source/Runtime/VectorVM/Private/VectorVMExperimental_Optimizer.inl
new file mode 100644
--- /dev/null
+++ b/Engine/Source/Runtime/VectorVM/Private/VectorVMExperimental_Optimizer.inl
@@ -0,0 +1,2303 @@
+// Copyright Epic Games, Inc. All Rights Reserved.
+
+/*
+	OptimizeVectorVMScript() takes the original VM's bytecode as input and outputs a new bytecode,
+a ConstRemapTable, an External Function Table and computes the number of TempRegs and ConstBuffs 
+required.  This function is not particularly efficient, but I don't think it needs to be, (I would 
+advise against optimizing it, as clarity is *MUCH* more important in this case imo).
+
+	Unlike in the original VM, this optimized bytecode can be saved and used in cooked builds.  
+There's no reason to keep the original bytecode around other than for editing.  This function 
+effectively acts as a back-end compiler, using the original VM's bytecode as an intermediate 
+representation.
+
+	The FVectorVMOptimizeContext has a struct called "Intermediate" which is some internal state that 
+	the optimizer needs for the duration of the OptimizeVectorVMScript().  The intermediate structure 
+	is usually free'd at the end of the OptimizeVectorVMScript() function, however it can be saved by 
+	passing VVMOptFlag_SaveIntermediateState in the Flags argument.  You may want to save it for 
+	debugging purposes, but there's no reason to save it during normal runtime execution in UE.
+
+	This document will list each step the Optimizer takes, what it does, roughly how it does it and 
+why.
+
+1. Create an Intermediate representation of all Instructions.
+	- Parse the input bytecode and generate an array of FVectorVMOptimizeInstructions.
+	- Any ConstBuff operands are saved in the OptimizeContext->ConstRemap table
+	- Any TempRegs operands are saved in the OptimizeContext->Intermediate.RegisterUsageBuffer 
+	  table
+	- Count the number of External functions
+
+	Instructions that have operands store an index: RegPtrOffset.  RegPtrOffset serves as a lookup 
+into the OptimizeContext->Intermediate.RegisterUsageBuffer table to see which TempRegs or 
+ConstBuffs an instruction uses.
+
+2. Alloc memory for the External Function Table and set the number of Inputs and Outputs each 
+function requires.  The function pointer is left NULL forever.  This External Function Table gets 
+copied to the FVectorVMState in VectorVMInit() where the function pointers are set in the 
+FVectorVMState for runtime execution.
+
+3. Perform some sanity checks: verify the ConstRemapTable is correct.  There's two parallel arrays 
+in ConstRemapTable.  The first maps the original sparse ConstBuff index to the new tightly packed 
+index.  The second array is the reverse mapping.
+
+4. Setup additional buffers:
+	- OptimizeContext->Intermedate.SSARegisterUsageBuffer - This is a parallel array to the 
+      RegisterUsageBuffer.  The RegPtrOffsets stored in each FVectorVMOptimizeInstruction serve as 
+	  in index into the SSA SSARegisterUsageBuffer as well.
+	- OptContext->Intermediate.InputRegisterFuseBuffer - uses the same RegPtrOffsetes as the two 
+      RegisterUsageBuffer.  This represents the index into the 
+	  OptimizeContext->Intermediate.Instructions array that a particular operand should replace 
+	  with an input instruction.  ie. an add instruction (AddIns) should replace operand 1's 
+	  TempReg with an Input instruction that's 8th in the list: 
+	  OptContext->Intermediate.InputRegisterFuseBuffer[AddIns->Op.RegPtrOffset + 1] = 8;
+
+5. Fill the SSARegisterUsageBuffer.  Loop through the Intermediate.Instructions array and fill out 
+the SSARegisterUsageBuffer with a single static assignment register index.  Each output for an 
+instruction gets a new SSA register.
+
+6. Input Fusing.  Loop through the Intermediate.Instructions array and find all the places where an
+Op's operands' TempRegs can be replaced with an Input instruction.  Set the InputFuseBits to a 
+bitfield of which operands can be replaced.
+	In the original VM Inputs into TempRegs are often directly written to an output.  This is 
+inefficient.  The new VM has a copy_to_output instruction.  If an input can be copied directly to 
+an output it's figured out in this step, and FVectorVMInstruction::Output.CopyFromInputInsIdx is 
+set.  Keep track of which Inputs are still required after this step.  Most input instructions will 
+not make it into the final bytecode.
+
+7. Remove unnecessary instructions.  Occasionally the original VM emitted instructions with outputs 
+that are never used.  We removed those here.
+
+8. Fixup SSA registers for removed instructions.  Instructions that get removed might have 
+unnecessary registers taking a name in the SSA table.  They are removed here.
+
+9. Re-order acquireindex instructions to be executed ASAP.  In order to minimize the state of the 
+VM while executing we want to output data we no longer need ASAP to free up the TempRegs for other 
+instructions.  The first thing we need to do is figure out which instances we're keeping, and which
+we're discarding.  We re-order the acquireindex instructions and their dependenices to happen ASAP.
+
+10. Re-order update_id instructions to happen as soon after the acquireindex instruction as 
+possible.  update_id uses acquireindex to free persistent IDs.
+
+11. Re-order the output instructions to happen ASAP.  The output instructions are placed 
+immediately after the register is last used.  We use the SSA table to figure this out.
+
+12. Re-order instructions that have no dependenices to immediately before their output is used.  
+This creates a tighter packing of SSA registers and allows registers to be freed for re-use.
+
+13. Re-order Input instructions directly before they're used.  Inputs to external functions and 
+acquire_id are not fused (maybe we could add this?).  This step places them immediately before 
+they're used.
+
+14. Group and sort all Output instructions that copy directly from an Input.  The copy_from_input 
+instruction takes a count parameter and will loop over each Input to copy during execution.
+
+15. Group all "normal" output instructions.  There's several new output_batch instructions that 
+can output more than one register at a time.  In order to write them efficiently the output 
+instructions are sorted here.
+
+16. Since we've re-ordered instructions the 
+OptimizeContext->Intermediate.InputRegisterFuseBuffer indices are wrong.  This step corrects all 
+Input instruction references to their new Index.
+
+17. Use the SSA registers to compute the minimum set of registers required to execute this script.
+This step writes the new minimized index back into the 
+OptContext->Intermediate.RegisterUsageBuffer array.  An instructions output TempReg can now alias 
+with its input.
+
+18. Write the final optimized bytecode.  This loops over the Instruction array twice.  The first 
+time counts how may bytes are required, the second pass actually writes the bytecode.  
+Instructions with fused inputs write two instruction: a fused_inputX_y and the operation itself.  
+Outputs can write either a copy_to_output, output_batch or a regular output instruction.
+
+*/
+
+#define VVM_OPT_MAX_REGS_PER_INS 256 //this is absurdly high, but still only 512 bytes on the stack
+
+struct VVMOptRAIIPtrToFree
+{ //automatically frees memory when they go out of scope
+	VVMOptRAIIPtrToFree(FVectorVMOptimizeContext *Ctx, void *Ptr) : Ctx(Ctx), Ptr(Ptr) { }
+	~VVMOptRAIIPtrToFree()
+	{
+		Ctx->Init.FreeFn(Ptr, __FILE__, __LINE__);
+	}
+	FVectorVMOptimizeContext *Ctx;
+	void *Ptr;
+};
+
+struct FVectorVMOptimizeInsRegUsage
+{
+	uint16  RegIndices[VVM_OPT_MAX_REGS_PER_INS]; //Index into FVectorVMOptimizeContext::RegisterUsageBuffer.  Output follows input
+	int     NumInputRegisters;
+	int     NumOutputRegisters;
+};
+
+static void VectorVMFreeOptimizerIntermediateData(FVectorVMOptimizeContext *OptContext)
+{
+	if (OptContext->Init.FreeFn)
+	{
+		OptContext->Init.FreeFn(OptContext->Intermediate.Instructions           , __FILE__, __LINE__);
+		OptContext->Init.FreeFn(OptContext->Intermediate.RegisterUsageBuffer    , __FILE__, __LINE__);
+		OptContext->Init.FreeFn(OptContext->Intermediate.SSARegisterUsageBuffer , __FILE__, __LINE__);
+		OptContext->Init.FreeFn(OptContext->Intermediate.InputRegisterFuseBuffer, __FILE__, __LINE__);
+		FMemory::Memset(&OptContext->Intermediate, 0, sizeof(OptContext->Intermediate));
+	}
+	else 
+	{
+		check(OptContext->Intermediate.Instructions            == nullptr);
+		check(OptContext->Intermediate.RegisterUsageBuffer     == nullptr);
+		check(OptContext->Intermediate.SSARegisterUsageBuffer  == nullptr);
+		check(OptContext->Intermediate.InputRegisterFuseBuffer == nullptr);
+	}
+}
+
+VECTORVM_API void FreeVectorVMOptimizeContext(FVectorVMOptimizeContext *OptContext)
+{
+	//save init data
+	VectorVMReallocFn *ReallocFn = OptContext->Init.ReallocFn;
+	VectorVMFreeFn    *FreeFn    = OptContext->Init.FreeFn;
+	//save error data
+	uint32 ErrorFlags                              = OptContext->Error.Flags;
+	uint32 ErrorLine                               = OptContext->Error.Line;
+	VectorVMOptimizeErrorCallback *ErrorCallbackFn = OptContext->Error.CallbackFn;
+	//free and zero everything
+	if (FreeFn)
+	{
+		FreeFn(OptContext->OutputBytecode, __FILE__, __LINE__);
+		FreeFn(OptContext->ConstRemap[0] , __FILE__, __LINE__);
+		FreeFn(OptContext->ConstRemap[1] , __FILE__, __LINE__);
+		FreeFn(OptContext->ExtFnTable    , __FILE__, __LINE__);
+	}
+	else
+	{
+		check(OptContext->OutputBytecode == nullptr);
+		check(OptContext->ConstRemap[0]  == nullptr);
+		check(OptContext->ConstRemap[1]  == nullptr);
+		check(OptContext->ExtFnTable     == nullptr);
+	}
+	VectorVMFreeOptimizerIntermediateData(OptContext);
+	FMemory::Memset(OptContext, 0, sizeof(FVectorVMOptimizeContext));
+	//restore init data
+	OptContext->Init.ReallocFn   = ReallocFn;
+	OptContext->Init.FreeFn      = FreeFn;
+	//restore error data
+	OptContext->Error.Flags      = ErrorFlags;
+	OptContext->Error.Line       = ErrorLine;
+	OptContext->Error.CallbackFn = ErrorCallbackFn;
+}
+
+static uint32 VectorVMOptimizerSetError_(FVectorVMOptimizeContext *OptContext, uint32 Flags, uint32 LineNum)
+{
+	OptContext->Error.Line = LineNum;
+	if (OptContext->Error.CallbackFn)
+	{
+		OptContext->Error.Flags = OptContext->Error.CallbackFn(OptContext, OptContext->Error.Flags | Flags);
+	}
+	else
+	{
+		OptContext->Error.Flags |= Flags;
+	}
+	if (OptContext->Error.Flags & VVMOptErr_Fatal)
+	{
+		check(false); //hit the debugger
+		FreeVectorVMOptimizeContext(OptContext);
+	}
+	return OptContext->Error.Flags;
+}
+
+#define VectorVMOptimizerSetError(Context, Flags)   VectorVMOptimizerSetError_(Context, Flags, __LINE__)
+
+static uint16 VectorVMOptimizeRemapConst(FVectorVMOptimizeContext *OptContext, uint16 ConstIdx)
+{
+	if (ConstIdx >= OptContext->NumConstsAlloced)
+	{
+		uint16 NumConstsToAlloc = (ConstIdx + 1 + 63) & ~63;
+		uint16 *ConstRemap0 = (uint16 *)OptContext->Init.ReallocFn(OptContext->ConstRemap[0], sizeof(uint16) * NumConstsToAlloc, __FILE__, __LINE__);
+		if (ConstRemap0 == nullptr)
+		{
+			VectorVMOptimizerSetError(OptContext, VVMOptErr_OutOfMemory | VVMOptErr_ConstRemap | VVMOptErr_Fatal);
+			return 0;
+		}
+		OptContext->ConstRemap[0] = ConstRemap0;
+		uint16 *ConstRemap1 = (uint16 *)OptContext->Init.ReallocFn(OptContext->ConstRemap[1], sizeof(uint16) * NumConstsToAlloc, __FILE__, __LINE__);
+		if (ConstRemap1 == nullptr)
+		{
+			VectorVMOptimizerSetError(OptContext, VVMOptErr_OutOfMemory | VVMOptErr_ConstRemap | VVMOptErr_Fatal);
+			return 0;
+		}
+		OptContext->ConstRemap[1] = ConstRemap1;
+		if (NumConstsToAlloc > OptContext->NumConstsAlloced)
+		{
+			FMemory::Memset(OptContext->ConstRemap[0] + OptContext->NumConstsAlloced, 0xFF, sizeof(uint16) * (NumConstsToAlloc - OptContext->NumConstsAlloced));
+			FMemory::Memset(OptContext->ConstRemap[1] + OptContext->NumConstsAlloced, 0xFF, sizeof(uint16) * (NumConstsToAlloc - OptContext->NumConstsAlloced));
+		}
+		OptContext->NumConstsAlloced = NumConstsToAlloc;
+	}
+	if (OptContext->ConstRemap[0][ConstIdx] == 0xFFFF)
+	{
+		OptContext->ConstRemap[0][ConstIdx] = OptContext->NumConstsRemapped;
+		OptContext->ConstRemap[1][OptContext->NumConstsRemapped] = ConstIdx;
+		OptContext->NumConstsRemapped++;
+		check(OptContext->NumConstsRemapped <= OptContext->NumConstsAlloced);
+	}
+	else
+	{
+		check(OptContext->ConstRemap[1][OptContext->ConstRemap[0][ConstIdx]] == ConstIdx);
+	}
+	return OptContext->ConstRemap[0][ConstIdx];
+}
+
+static int GetRegistersUsedForInstruction(FVectorVMOptimizeContext *OptContext, FVectorVMOptimizeInstruction *Ins, FVectorVMOptimizeInsRegUsage *OutRegUsage) {
+	OutRegUsage->NumInputRegisters  = 0;
+	OutRegUsage->NumOutputRegisters = 0;
+	switch (Ins->OpCat)
+	{
+	case EVectorVMOpCategory::Input:
+		if (Ins->Input.FirstInsInsertIdx != -1)
+		{
+			OutRegUsage->RegIndices[OutRegUsage->NumOutputRegisters++] = Ins->Input.DstRegPtrOffset;
+		}
+		break;
+	case EVectorVMOpCategory::Output:
+		if (!(OptContext->Intermediate.RegisterUsageBuffer[Ins->Output.RegPtrOffset] & 0x8000))
+		{
+			OutRegUsage->RegIndices[OutRegUsage->NumInputRegisters++] = Ins->Output.RegPtrOffset;
+		}
+		if (Ins->OpCode != EVectorVMOp::copy_to_output)
+		{
+			if (!(OptContext->Intermediate.RegisterUsageBuffer[Ins->Output.RegPtrOffset + 1] & 0x8000))
+			{
+				OutRegUsage->RegIndices[OutRegUsage->NumInputRegisters++] = Ins->Output.RegPtrOffset + 1;
+			}
+		}
+		break;
+	case EVectorVMOpCategory::Op:
+		if (Ins->Op.InputFuseBits == 0) //all inputs are regular registers
+		{
+			{ //Input registers
+				int InputCount = 0;
+				for (int i = 0; i < Ins->Op.NumInputs; ++i) {
+					if (!(OptContext->Intermediate.RegisterUsageBuffer[Ins->Op.RegPtrOffset + i] & 0x8000)) {
+						OutRegUsage->RegIndices[InputCount++] = Ins->Op.RegPtrOffset + i;
+					}
+				}
+				OutRegUsage->NumInputRegisters = InputCount;
+			}
+			{ //Output registers
+				int OutputCount = 0;
+				for (int i = 0; i < Ins->Op.NumOutputs; ++i) {
+					OutRegUsage->RegIndices[OutRegUsage->NumInputRegisters + i] = Ins->Op.RegPtrOffset + Ins->Op.NumInputs + i;
+				}
+				OutRegUsage->NumOutputRegisters = Ins->Op.NumOutputs;
+			}
+		}
+		else //at least one of the inputs should be from a dataset, not a register
+		{ 
+			check(Ins->Op.NumInputs > 0);
+			check(Ins->Op.NumInputs <= 3);
+			{ //Input registers
+				int InputCount = 0;
+				for (int i = 0; i < Ins->Op.NumInputs; ++i) {
+					if ((!(Ins->Op.InputFuseBits & (1 << i))) && !(OptContext->Intermediate.RegisterUsageBuffer[Ins->Op.RegPtrOffset + i] & 0x8000)) {
+						OutRegUsage->RegIndices[InputCount] = Ins->Op.RegPtrOffset + i;
+						++InputCount;
+					}
+				}
+				OutRegUsage->NumInputRegisters = InputCount;
+			}
+			{ //Output Registers
+				for (int i = 0; i < Ins->Op.NumOutputs; ++i) {
+					OutRegUsage->RegIndices[OutRegUsage->NumInputRegisters + i] = Ins->Op.RegPtrOffset + Ins->Op.NumInputs + i;
+				}
+				OutRegUsage->NumOutputRegisters = Ins->Op.NumOutputs;
+			}
+		}
+		break;
+	case EVectorVMOpCategory::ExtFnCall:
+		check(Ins->ExtFnCall.NumInputs + Ins->ExtFnCall.NumOutputs < VVM_OPT_MAX_REGS_PER_INS);
+		//if this check fails (*EXTREMELY* unlikely), just increase VVM_OPT_MAX_REGS_PER_INS
+		for (int i = 0; i < Ins->ExtFnCall.NumInputs; ++i)
+		{
+			if ((OptContext->Intermediate.RegisterUsageBuffer[Ins->ExtFnCall.RegPtrOffset + i] & 0x8000) == 0)
+			{
+				OutRegUsage->RegIndices[OutRegUsage->NumInputRegisters++] = Ins->ExtFnCall.RegPtrOffset + i;
+			}
+		}
+		for (int i = 0; i < Ins->ExtFnCall.NumOutputs; ++i)
+		{
+			OutRegUsage->RegIndices[OutRegUsage->NumInputRegisters + i] = Ins->ExtFnCall.RegPtrOffset + Ins->ExtFnCall.NumInputs + i;
+		}
+		OutRegUsage->NumOutputRegisters = Ins->ExtFnCall.NumOutputs;
+		break;
+	case EVectorVMOpCategory::IndexGen:
+		if (!(OptContext->Intermediate.RegisterUsageBuffer[Ins->IndexGen.RegPtrOffset + 0] & 0x8000))
+		{
+			OutRegUsage->RegIndices[OutRegUsage->NumInputRegisters++] = Ins->IndexGen.RegPtrOffset + 0;
+		}
+		OutRegUsage->RegIndices[OutRegUsage->NumInputRegisters + OutRegUsage->NumOutputRegisters++] = Ins->IndexGen.RegPtrOffset + 1;
+		break;
+	case EVectorVMOpCategory::ExecIndex:
+		if (!(OptContext->Intermediate.RegisterUsageBuffer[Ins->IndexGen.RegPtrOffset + 0] & 0x8000))
+		{
+			OutRegUsage->RegIndices[OutRegUsage->NumOutputRegisters++] = Ins->ExecIndex.RegPtrOffset;
+		}
+		break;
+	case EVectorVMOpCategory::RWBuffer:
+		OutRegUsage->RegIndices[0] = Ins->RWBuffer.RegPtrOffset + 0;
+		OutRegUsage->RegIndices[1] = Ins->RWBuffer.RegPtrOffset + 1;
+		if (Ins->OpCode == EVectorVMOp::acquire_id)
+		{
+			OutRegUsage->NumOutputRegisters = 2;
+		} 
+		else if (Ins->OpCode == EVectorVMOp::update_id)
+		{
+			OutRegUsage->NumInputRegisters = 2;
+		}
+		else
+		{
+			check(false);
+		}
+		break;
+	case EVectorVMOpCategory::Stat:
+		break;
+	case EVectorVMOpCategory::Fused:
+		check(false); //we don't write an intermediate representation of a fused instruction, so this shouldn't be possible
+	}
+	check(OutRegUsage->NumInputRegisters + OutRegUsage->NumOutputRegisters < VVM_OPT_MAX_REGS_PER_INS);
+	return OutRegUsage->NumInputRegisters + OutRegUsage->NumOutputRegisters;
+}
+
+static int GetInstructionDependencyChain(FVectorVMOptimizeContext *OptContext, int InsIdxToCheck, int *RegToCheckStack, int *InstructionIdxStack)
+{
+	int NumRegistersToCheck           = 0;
+	int NumInstructions               = 0;
+	FVectorVMOptimizeInstruction *Ins = OptContext->Intermediate.Instructions + InsIdxToCheck;
+	FVectorVMOptimizeInsRegUsage InsRegUse = { };
+	FVectorVMOptimizeInsRegUsage OpRegUse  = { };
+
+	GetRegistersUsedForInstruction(OptContext, Ins, &InsRegUse);
+	for (int i = 0; i < InsRegUse.NumInputRegisters; ++i)
+	{
+		RegToCheckStack[NumRegistersToCheck++] = OptContext->Intermediate.SSARegisterUsageBuffer[InsRegUse.RegIndices[i]];
+	}
+	while (NumRegistersToCheck > 0)
+	{
+		uint16 RegToCheck = RegToCheckStack[--NumRegistersToCheck];
+		for (int InsIdx = InsIdxToCheck - 1; InsIdx >= 0; --InsIdx)
+		{
+			GetRegistersUsedForInstruction(OptContext, OptContext->Intermediate.Instructions + InsIdx, &OpRegUse);
+			for (int j = 0; j < OpRegUse.NumOutputRegisters; ++j)
+			{
+				uint16 OutputReg = OptContext->Intermediate.SSARegisterUsageBuffer[OpRegUse.RegIndices[OpRegUse.NumInputRegisters + j]];
+				if (RegToCheck == OutputReg)
+				{
+					bool InsAlreadyInStack = false;
+					for (int i = 0; i < NumInstructions; ++i)
+					{
+						if (InstructionIdxStack[i] == InsIdx)
+						{
+							InsAlreadyInStack = true;
+							break;
+						}
+					}
+					if (!InsAlreadyInStack)
+					{
+						{ //insert in sorted low-to-high order
+							int InsertionSlot = NumInstructions;
+							for (int i = 0; i < NumInstructions; ++i)
+							{
+								if (InsIdx < InstructionIdxStack[i])
+								{
+									InsertionSlot = i;
+									FMemory::Memmove(InstructionIdxStack + InsertionSlot + 1, InstructionIdxStack + InsertionSlot, sizeof(int) * (NumInstructions - InsertionSlot));
+									break;
+								}
+							}
+							InstructionIdxStack[InsertionSlot] = InsIdx;
+							++NumInstructions;
+						}
+						for (int k = 0; k < OpRegUse.NumInputRegisters; ++k)
+						{
+							bool RegAlreadyInStack = false;
+							uint16 Reg = OptContext->Intermediate.SSARegisterUsageBuffer[OpRegUse.RegIndices[k]];
+							for (int i = 0; i < NumRegistersToCheck; ++i)
+							{
+								if (RegToCheckStack[i] == Reg)
+								{
+									RegAlreadyInStack = true;
+									break;
+								}
+							}
+							if (!RegAlreadyInStack)
+							{
+								RegToCheckStack[NumRegistersToCheck++] = Reg;
+							}
+						}
+					}
+				}
+			}
+		}
+	}
+	return NumInstructions;
+}
+
+static EVectorVMOpCategory GetOpCategoryFromOp(EVectorVMOp op)
+{
+	switch (op)
+	{
+		case EVectorVMOp::done:                         return EVectorVMOpCategory::Other;
+		case EVectorVMOp::add:                          return EVectorVMOpCategory::Op;
+		case EVectorVMOp::sub:                          return EVectorVMOpCategory::Op;
+		case EVectorVMOp::mul:                          return EVectorVMOpCategory::Op;
+		case EVectorVMOp::div:                          return EVectorVMOpCategory::Op;
+		case EVectorVMOp::mad:                          return EVectorVMOpCategory::Op;
+		case EVectorVMOp::lerp:                         return EVectorVMOpCategory::Op;
+		case EVectorVMOp::rcp:                          return EVectorVMOpCategory::Op;
+		case EVectorVMOp::rsq:                          return EVectorVMOpCategory::Op;
+		case EVectorVMOp::sqrt:                         return EVectorVMOpCategory::Op;
+		case EVectorVMOp::neg:                          return EVectorVMOpCategory::Op;
+		case EVectorVMOp::abs:                          return EVectorVMOpCategory::Op;
+		case EVectorVMOp::exp:                          return EVectorVMOpCategory::Op;
+		case EVectorVMOp::exp2:                         return EVectorVMOpCategory::Op;
+		case EVectorVMOp::log:                          return EVectorVMOpCategory::Op;
+		case EVectorVMOp::log2:                         return EVectorVMOpCategory::Op;
+		case EVectorVMOp::sin:                          return EVectorVMOpCategory::Op;
+		case EVectorVMOp::cos:                          return EVectorVMOpCategory::Op;
+		case EVectorVMOp::tan:                          return EVectorVMOpCategory::Op;
+		case EVectorVMOp::asin:                         return EVectorVMOpCategory::Op;
+		case EVectorVMOp::acos:                         return EVectorVMOpCategory::Op;
+		case EVectorVMOp::atan:                         return EVectorVMOpCategory::Op;
+		case EVectorVMOp::atan2:                        return EVectorVMOpCategory::Op;
+		case EVectorVMOp::ceil:                         return EVectorVMOpCategory::Op;
+		case EVectorVMOp::floor:                        return EVectorVMOpCategory::Op;
+		case EVectorVMOp::fmod:                         return EVectorVMOpCategory::Op;
+		case EVectorVMOp::frac:                         return EVectorVMOpCategory::Op;
+		case EVectorVMOp::trunc:                        return EVectorVMOpCategory::Op;
+		case EVectorVMOp::clamp:                        return EVectorVMOpCategory::Op;
+		case EVectorVMOp::min:                          return EVectorVMOpCategory::Op;
+		case EVectorVMOp::max:                          return EVectorVMOpCategory::Op;
+		case EVectorVMOp::pow:                          return EVectorVMOpCategory::Op;
+		case EVectorVMOp::round:                        return EVectorVMOpCategory::Op;
+		case EVectorVMOp::sign:                         return EVectorVMOpCategory::Op;
+		case EVectorVMOp::step:                         return EVectorVMOpCategory::Op;
+		case EVectorVMOp::random:                       return EVectorVMOpCategory::Op;
+		case EVectorVMOp::noise:                        return EVectorVMOpCategory::Op;
+		case EVectorVMOp::cmplt:                        return EVectorVMOpCategory::Op;
+		case EVectorVMOp::cmple:                        return EVectorVMOpCategory::Op;
+		case EVectorVMOp::cmpgt:                        return EVectorVMOpCategory::Op;
+		case EVectorVMOp::cmpge:                        return EVectorVMOpCategory::Op;
+		case EVectorVMOp::cmpeq:                        return EVectorVMOpCategory::Op;
+		case EVectorVMOp::cmpneq:                       return EVectorVMOpCategory::Op;
+		case EVectorVMOp::select:                       return EVectorVMOpCategory::Op;
+		case EVectorVMOp::addi:                         return EVectorVMOpCategory::Op;
+		case EVectorVMOp::subi:                         return EVectorVMOpCategory::Op;
+		case EVectorVMOp::muli:                         return EVectorVMOpCategory::Op;
+		case EVectorVMOp::divi:                         return EVectorVMOpCategory::Op;
+		case EVectorVMOp::clampi:                       return EVectorVMOpCategory::Op;
+		case EVectorVMOp::mini:                         return EVectorVMOpCategory::Op;
+		case EVectorVMOp::maxi:                         return EVectorVMOpCategory::Op;
+		case EVectorVMOp::absi:                         return EVectorVMOpCategory::Op;
+		case EVectorVMOp::negi:                         return EVectorVMOpCategory::Op;
+		case EVectorVMOp::signi:                        return EVectorVMOpCategory::Op;
+		case EVectorVMOp::randomi:                      return EVectorVMOpCategory::Op;
+		case EVectorVMOp::cmplti:                       return EVectorVMOpCategory::Op;
+		case EVectorVMOp::cmplei:                       return EVectorVMOpCategory::Op;
+		case EVectorVMOp::cmpgti:                       return EVectorVMOpCategory::Op;
+		case EVectorVMOp::cmpgei:                       return EVectorVMOpCategory::Op;
+		case EVectorVMOp::cmpeqi:                       return EVectorVMOpCategory::Op;
+		case EVectorVMOp::cmpneqi:                      return EVectorVMOpCategory::Op;
+		case EVectorVMOp::bit_and:                      return EVectorVMOpCategory::Op;
+		case EVectorVMOp::bit_or:                       return EVectorVMOpCategory::Op;
+		case EVectorVMOp::bit_xor:                      return EVectorVMOpCategory::Op;
+		case EVectorVMOp::bit_not:                      return EVectorVMOpCategory::Op;
+		case EVectorVMOp::bit_lshift:                   return EVectorVMOpCategory::Op;
+		case EVectorVMOp::bit_rshift:                   return EVectorVMOpCategory::Op;
+		case EVectorVMOp::logic_and:                    return EVectorVMOpCategory::Op;
+		case EVectorVMOp::logic_or:                     return EVectorVMOpCategory::Op;
+		case EVectorVMOp::logic_xor:                    return EVectorVMOpCategory::Op;
+		case EVectorVMOp::logic_not:                    return EVectorVMOpCategory::Op;
+		case EVectorVMOp::f2i:                          return EVectorVMOpCategory::Op;
+		case EVectorVMOp::i2f:                          return EVectorVMOpCategory::Op;
+		case EVectorVMOp::f2b:                          return EVectorVMOpCategory::Op;
+		case EVectorVMOp::b2f:                          return EVectorVMOpCategory::Op;
+		case EVectorVMOp::i2b:                          return EVectorVMOpCategory::Op;
+		case EVectorVMOp::b2i:                          return EVectorVMOpCategory::Op;
+		case EVectorVMOp::inputdata_float:              return EVectorVMOpCategory::Input;
+		case EVectorVMOp::inputdata_int32:              return EVectorVMOpCategory::Input;
+		case EVectorVMOp::inputdata_half:               return EVectorVMOpCategory::Input;
+		case EVectorVMOp::inputdata_noadvance_float:    return EVectorVMOpCategory::Input;
+		case EVectorVMOp::inputdata_noadvance_int32:    return EVectorVMOpCategory::Input;
+		case EVectorVMOp::inputdata_noadvance_half:     return EVectorVMOpCategory::Input;
+		case EVectorVMOp::outputdata_float:             return EVectorVMOpCategory::Output;
+		case EVectorVMOp::outputdata_int32:             return EVectorVMOpCategory::Output;
+		case EVectorVMOp::outputdata_half:              return EVectorVMOpCategory::Output;
+		case EVectorVMOp::acquireindex:                 return EVectorVMOpCategory::IndexGen;
+		case EVectorVMOp::external_func_call:           return EVectorVMOpCategory::ExtFnCall;
+		case EVectorVMOp::exec_index:                   return EVectorVMOpCategory::ExecIndex;
+		case EVectorVMOp::noise2D:                      return EVectorVMOpCategory::Other;
+		case EVectorVMOp::noise3D:                      return EVectorVMOpCategory::Other;
+		case EVectorVMOp::enter_stat_scope:             return EVectorVMOpCategory::Stat;
+		case EVectorVMOp::exit_stat_scope:              return EVectorVMOpCategory::Stat;
+		case EVectorVMOp::update_id:                    return EVectorVMOpCategory::RWBuffer;
+		case EVectorVMOp::acquire_id:                   return EVectorVMOpCategory::RWBuffer;
+		case EVectorVMOp::fused_input1_1:               return EVectorVMOpCategory::Fused;
+		case EVectorVMOp::fused_input2_1:               return EVectorVMOpCategory::Fused;
+		case EVectorVMOp::fused_input2_2:               return EVectorVMOpCategory::Fused;
+		case EVectorVMOp::fused_input2_3:               return EVectorVMOpCategory::Fused;
+		case EVectorVMOp::fused_input3_1:               return EVectorVMOpCategory::Fused;
+		case EVectorVMOp::fused_input3_2:               return EVectorVMOpCategory::Fused;
+		case EVectorVMOp::fused_input3_4:               return EVectorVMOpCategory::Fused;
+		case EVectorVMOp::fused_input3_3:               return EVectorVMOpCategory::Fused;
+		case EVectorVMOp::fused_input3_5:               return EVectorVMOpCategory::Fused;
+		case EVectorVMOp::fused_input3_6:               return EVectorVMOpCategory::Fused;
+		case EVectorVMOp::fused_input3_7:               return EVectorVMOpCategory::Fused;
+		case EVectorVMOp::copy_to_output:               return EVectorVMOpCategory::Output;
+		case EVectorVMOp::output_batch2:                return EVectorVMOpCategory::Output;
+		case EVectorVMOp::output_batch3:                return EVectorVMOpCategory::Output;
+		case EVectorVMOp::output_batch4:                return EVectorVMOpCategory::Output;
+		case EVectorVMOp::output_batch7:                return EVectorVMOpCategory::Output;
+		case EVectorVMOp::output_batch8:                return EVectorVMOpCategory::Output;
+		default:                          check(false); return EVectorVMOpCategory::Other;
+	}
+}
+
+inline uint64 VVMCopyToOutputInsGetSortKey(FVectorVMOptimizeInstruction *Instructions, FVectorVMOptimizeInstruction *OutputIns)
+{
+	check(OutputIns->OpCat == EVectorVMOpCategory::Output);
+	check(OutputIns->Output.CopyFromInputInsIdx != -1);
+	check(OutputIns->Output.DataSetIdx < (1 << 14));                                                // max 14 bits for DataSet Index (In reality this number is < 5... ie 3 bits)
+	FVectorVMOptimizeInstruction *InputIns = Instructions + OutputIns->Output.CopyFromInputInsIdx;																					
+	check(InputIns->OpCat == EVectorVMOpCategory::Input);
+	check(InputIns->Input.DataSetIdx < (1 << 14));	                                                // max 14 bits for DataSet Index (In reality this number is < 5... ie 3 bits)
+	uint8 InputRegType  = (uint8)InputIns->OpCode - (uint8)EVectorVMOp::inputdata_float;
+	uint8 OutputRegType = (uint8)OutputIns->OpCode - (uint8)EVectorVMOp::outputdata_float;
+	check(InputRegType == OutputRegType);                                                           // input and output reg type should match, so we only use 1 bit
+	check(OutputRegType == 1 || OutputRegType == 0);									            // if they ever don't match (WHY?!) then we can change this to use 2 bits 
+	uint64 key = ((uint64)(OutputRegType & 1)          << 63ULL) +                                  // 63    - Float/Int flag
+	             ((uint64)OutputIns->Output.DataSetIdx << 49ULL) +                                  // 49-62 - Output Data Set Index
+	             ((uint64)InputIns->Input.DataSetIdx   << 35ULL) +                                  // 35-49 - Input Data Set Index
+	             ((uint64)InputIns->Input.InputIdx     << 16ULL) +                                  // 16-31 - Input Src
+	             ((uint64)OutputIns->Output.DstRegIdx)           ;                                  // 0-15  - Output Dest
+	return key;
+}
+
+inline uint64 VVMOutputInsGetSortKey(uint16 *SSARegisters, FVectorVMOptimizeInstruction *OutputIns)
+{
+	check(OutputIns->OpCat == EVectorVMOpCategory::Output);
+	check(OutputIns->Output.DataSetIdx < (1 << 14));                                                // max 14 bits for DataSet Index (In reality this number is < 5... ie 3 bits)
+	check(OutputIns->Output.CopyFromInputInsIdx == -1);
+	check((int)OutputIns->OpCode >= (int)EVectorVMOp::outputdata_float);
+	uint64 key = (((uint64)OutputIns->OpCode - (uint64)EVectorVMOp::outputdata_float) << 62ULL) +
+		         ((uint64)OutputIns->Output.DataSetIdx                                << 48ULL) +
+		         ((uint64)SSARegisters[OutputIns->Output.RegPtrOffset]                << 16ULL) +
+		         ((uint64)OutputIns->Output.DstRegIdx)                                          ;
+	return key;
+}
+
+VECTORVM_API uint32 OptimizeVectorVMScript(const uint8 *InBytecode, int InBytecodeLen, FVectorVMExtFunctionData *ExtFnIOData, int NumExtFns, FVectorVMOptimizeContext *OptContext, uint32 Flags)
+{
+#define AllocRegisterUse(NumRegistersToAlloc)	if (OptContext->Intermediate.NumRegistersUsed + (NumRegistersToAlloc) >= NumRegisterUsageAlloced)                                                                               \
+												{                                                                                                                                                                               \
+													if (NumRegisterUsageAlloced == 0)                                                                                                                                           \
+													{                                                                                                                                                                           \
+														NumRegisterUsageAlloced = 32;                                                                                                                                           \
+													}                                                                                                                                                                           \
+													else                                                                                                                                                                        \
+													{                                                                                                                                                                           \
+														NumRegisterUsageAlloced <<= 1;                                                                                                                                          \
+													}                                                                                                                                                                           \
+													uint16 *NewRegisters = (uint16 *)OptContext->Init.ReallocFn(OptContext->Intermediate.RegisterUsageBuffer, sizeof(uint16) * NumRegisterUsageAlloced, __FILE__, __LINE__);    \
+													if (NewRegisters == nullptr)                                                                                                                                                \
+													{                                                                                                                                                                           \
+														return VectorVMOptimizerSetError(OptContext, VVMOptErr_OutOfMemory | VVMOptErr_RegisterUsage | VVMOptErr_Fatal);                                                        \
+													}                                                                                                                                                                           \
+													else                                                                                                                                                                        \
+													{                                                                                                                                                                           \
+														OptContext->Intermediate.RegisterUsageBuffer = NewRegisters;                                                                                                            \
+													}                                                                                                                                                                           \
+												}
+
+#define VVMOptimizeWriteRegIndex(OpIpVecIdx, io)	AllocRegisterUse(1);                                                                                                    \
+													if (*OpPtrIn & (1 << OpIpVecIdx))                                                                                       \
+													{                                                                                                                       \
+														check((VecIndices[OpIpVecIdx] & 3) == 0);                                                                           \
+														uint16 Idx = (VecIndices[OpIpVecIdx] >> 2);                                                                         \
+														uint16 RemappedIdx = VectorVMOptimizeRemapConst(OptContext, Idx) | 0x8000;                                          \
+														if (OptContext->Error.Flags & VVMOptErr_Fatal)                                                                      \
+														{                                                                                                                   \
+															return OptContext->Error.Flags;                                                                                 \
+														}                                                                                                                   \
+														if (Instruction->Op.NumInputs == 0 && Instruction->Op.NumOutputs == 0)                                              \
+														{                                                                                                                   \
+															Instruction->Op.RegPtrOffset = OptContext->Intermediate.NumRegistersUsed;                                       \
+														}                                                                                                                   \
+														++Instruction->Op.NumInputs;                                                                                        \
+														OptContext->Intermediate.RegisterUsageBuffer[OptContext->Intermediate.NumRegistersUsed] = RemappedIdx;              \
+													} else {                                                                                                                \
+														if (Instruction->Op.NumInputs == 0 && Instruction->Op.NumOutputs == 0)                                              \
+														{                                                                                                                   \
+															Instruction->Op.RegPtrOffset = OptContext->Intermediate.NumRegistersUsed;                                       \
+														}                                                                                                                   \
+														if (io)                                                                                                             \
+														{                                                                                                                   \
+															++Instruction->Op.NumOutputs;                                                                                   \
+														}                                                                                                                   \
+														else                                                                                                                \
+														{                                                                                                                   \
+															++Instruction->Op.NumInputs;                                                                                    \
+														}                                                                                                                   \
+														OptContext->Intermediate.RegisterUsageBuffer[OptContext->Intermediate.NumRegistersUsed] = VecIndices[OpIpVecIdx];   \
+													}                                                                                                                       \
+													++OptContext->Intermediate.NumRegistersUsed;                                                                            \
+													OptContext->Intermediate.NumBytecodeBytes += 2;
+
+
+#define VVMOptimizeVecIns1							VVMOptimizeWriteRegIndex(0, 0);		\
+													VVMOptimizeWriteRegIndex(1, 1);		\
+													OpPtrIn += 5;
+
+#define VVMOptimizeVecIns2							VVMOptimizeWriteRegIndex(0, 0);		\
+													VVMOptimizeWriteRegIndex(1, 0);		\
+													VVMOptimizeWriteRegIndex(2, 1);		\
+													OpPtrIn += 7;
+
+#define VVMOptimizeVecIns3							VVMOptimizeWriteRegIndex(0, 0);		\
+													VVMOptimizeWriteRegIndex(1, 0);		\
+													VVMOptimizeWriteRegIndex(2, 0);		\
+													VVMOptimizeWriteRegIndex(3, 1);		\
+													OpPtrIn += 9;
+	FreeVectorVMOptimizeContext(OptContext);
+	if (InBytecode == nullptr || InBytecodeLen == 0)
+	{
+		return 0;
+	}
+
+	if (OptContext->Init.ReallocFn == nullptr)
+	{
+		OptContext->Init.ReallocFn = VVMDefaultRealloc;
+	}
+	if (OptContext->Init.FreeFn == nullptr)
+	{
+		OptContext->Init.FreeFn = VVMDefaultFree;
+	}
+	OptContext->MaxExtFnUsed       = -1;
+
+	uint32 NumInstructionsAlloced  = 0;
+	uint32 NumBytecodeBytesAlloced = 0;
+	uint32 NumRegisterUsageAlloced = 0;
+	int32  MaxRWBufferUsed         = -1;
+
+	const uint8 *OpPtrIn = InBytecode;
+	const uint8 *OpPtrInEnd = InBytecode + InBytecodeLen;
+
+	//Step 1: Create Intermediate representation of all Instructions
+	while (OpPtrIn < OpPtrInEnd)
+	{
+		//alloc new instruction
+		check(OptContext->Intermediate.NumInstructions <= NumInstructionsAlloced);
+		if (OptContext->Intermediate.NumInstructions >= NumInstructionsAlloced)
+		{
+			if (NumInstructionsAlloced == 0)
+			{
+				NumInstructionsAlloced = 16;
+			}
+			else
+			{
+				NumInstructionsAlloced <<= 1;
+			}
+			FVectorVMOptimizeInstruction *NewInstructions = (FVectorVMOptimizeInstruction *)OptContext->Init.ReallocFn(OptContext->Intermediate.Instructions, sizeof(FVectorVMOptimizeInstruction) * NumInstructionsAlloced, __FILE__, __LINE__);
+			if (NewInstructions == nullptr)
+			{
+				return VectorVMOptimizerSetError(OptContext, VVMOptErr_OutOfMemory | VVMOptErr_Instructions | VVMOptErr_Fatal);
+			}
+			FMemory::Memset(NewInstructions + OptContext->Intermediate.NumInstructions, 0, sizeof(FVectorVMOptimizeInstruction) * (NumInstructionsAlloced - OptContext->Intermediate.NumInstructions));
+			OptContext->Intermediate.Instructions = NewInstructions;
+		}
+		uint16 *VecIndices	= (uint16 *)(OpPtrIn + 2);
+		EVectorVMOp op		= (EVectorVMOp)*OpPtrIn;
+
+		FVectorVMOptimizeInstruction *Instruction = OptContext->Intermediate.Instructions + OptContext->Intermediate.NumInstructions;
+		Instruction->Index = OptContext->Intermediate.NumInstructions++;
+		Instruction->OpCode = op;
+		Instruction->OpCat = GetOpCategoryFromOp(op);
+		Instruction->PtrOffsetInOrigBytecode = (uint32)(OpPtrIn - InBytecode);
+		OpPtrIn++;
+		switch (op)
+		{
+			case EVectorVMOp::done:                                                                 break;
+			case EVectorVMOp::add:                                  VVMOptimizeVecIns2;             break;
+			case EVectorVMOp::sub:                                  VVMOptimizeVecIns2;             break;
+			case EVectorVMOp::mul:                                  VVMOptimizeVecIns2;             break;
+			case EVectorVMOp::div:                                  VVMOptimizeVecIns2;             break;
+			case EVectorVMOp::mad:                                  VVMOptimizeVecIns3;             break;
+			case EVectorVMOp::lerp:                                 VVMOptimizeVecIns3;             break;
+			case EVectorVMOp::rcp:                                  VVMOptimizeVecIns1;             break;
+			case EVectorVMOp::rsq:                                  VVMOptimizeVecIns1;             break;
+			case EVectorVMOp::sqrt:                                 VVMOptimizeVecIns1;             break;
+			case EVectorVMOp::neg:                                  VVMOptimizeVecIns1;             break;
+			case EVectorVMOp::abs:                                  VVMOptimizeVecIns1;             break;
+			case EVectorVMOp::exp:                                  VVMOptimizeVecIns1;             break;
+			case EVectorVMOp::exp2:                                 VVMOptimizeVecIns1;             break;
+			case EVectorVMOp::log:                                  VVMOptimizeVecIns1;             break;
+			case EVectorVMOp::log2:                                 VVMOptimizeVecIns1;             break;
+			case EVectorVMOp::sin:                                  VVMOptimizeVecIns1;             break;
+			case EVectorVMOp::cos:                                  VVMOptimizeVecIns1;             break;
+			case EVectorVMOp::tan:                                  VVMOptimizeVecIns1;             break;
+			case EVectorVMOp::asin:                                 VVMOptimizeVecIns1;             break;
+			case EVectorVMOp::acos:                                 VVMOptimizeVecIns1;             break;
+			case EVectorVMOp::atan:                                 VVMOptimizeVecIns1;             break;
+			case EVectorVMOp::atan2:                                VVMOptimizeVecIns2;             break;
+			case EVectorVMOp::ceil:                                 VVMOptimizeVecIns1;             break;
+			case EVectorVMOp::floor:                                VVMOptimizeVecIns1;             break;
+			case EVectorVMOp::fmod:                                 VVMOptimizeVecIns2;             break;
+			case EVectorVMOp::frac:                                 VVMOptimizeVecIns1;             break;
+			case EVectorVMOp::trunc:                                VVMOptimizeVecIns1;             break;
+			case EVectorVMOp::clamp:                                VVMOptimizeVecIns3;             break;
+			case EVectorVMOp::min:                                  VVMOptimizeVecIns2;             break;
+			case EVectorVMOp::max:                                  VVMOptimizeVecIns2;             break;
+			case EVectorVMOp::pow:                                  VVMOptimizeVecIns2;             break;
+			case EVectorVMOp::round:                                VVMOptimizeVecIns1;             break;
+			case EVectorVMOp::sign:                                 VVMOptimizeVecIns1;             break;
+			case EVectorVMOp::step:                                 VVMOptimizeVecIns2;             break;
+			case EVectorVMOp::random:                               VVMOptimizeVecIns1;             break;
+			case EVectorVMOp::noise:                                check(false);                   break;
+			case EVectorVMOp::cmplt:                                VVMOptimizeVecIns2;             break;
+			case EVectorVMOp::cmple:                                VVMOptimizeVecIns2;             break;
+			case EVectorVMOp::cmpgt:                                VVMOptimizeVecIns2;             break;
+			case EVectorVMOp::cmpge:                                VVMOptimizeVecIns2;             break;
+			case EVectorVMOp::cmpeq:                                VVMOptimizeVecIns2;             break;
+			case EVectorVMOp::cmpneq:                               VVMOptimizeVecIns2;             break;
+			case EVectorVMOp::select:                               VVMOptimizeVecIns3;             break;
+			case EVectorVMOp::addi:                                 VVMOptimizeVecIns2;             break;
+			case EVectorVMOp::subi:                                 VVMOptimizeVecIns2;             break;
+			case EVectorVMOp::muli:                                 VVMOptimizeVecIns2;             break;
+			case EVectorVMOp::divi:                                 VVMOptimizeVecIns2;             break;
+			case EVectorVMOp::clampi:                               VVMOptimizeVecIns3;             break;
+			case EVectorVMOp::mini:                                 VVMOptimizeVecIns2;             break;
+			case EVectorVMOp::maxi:                                 VVMOptimizeVecIns2;             break;
+			case EVectorVMOp::absi:                                 VVMOptimizeVecIns1;             break;
+			case EVectorVMOp::negi:                                 VVMOptimizeVecIns1;             break;
+			case EVectorVMOp::signi:                                VVMOptimizeVecIns1;             break;
+			case EVectorVMOp::randomi:                              VVMOptimizeVecIns1;             break;
+			case EVectorVMOp::cmplti:                               VVMOptimizeVecIns2;             break;
+			case EVectorVMOp::cmplei:                               VVMOptimizeVecIns2;             break;
+			case EVectorVMOp::cmpgti:                               VVMOptimizeVecIns2;             break;
+			case EVectorVMOp::cmpgei:                               VVMOptimizeVecIns2;             break;
+			case EVectorVMOp::cmpeqi:                               VVMOptimizeVecIns2;             break;
+			case EVectorVMOp::cmpneqi:                              VVMOptimizeVecIns2;             break;
+			case EVectorVMOp::bit_and:                              VVMOptimizeVecIns2;             break;
+			case EVectorVMOp::bit_or:                               VVMOptimizeVecIns2;             break;
+			case EVectorVMOp::bit_xor:                              VVMOptimizeVecIns2;             break;
+			case EVectorVMOp::bit_not:                              VVMOptimizeVecIns1;             break;
+			case EVectorVMOp::bit_lshift:                           VVMOptimizeVecIns2;             break;
+			case EVectorVMOp::bit_rshift:                           VVMOptimizeVecIns2;             break;
+			case EVectorVMOp::logic_and:                            VVMOptimizeVecIns2;             break;
+			case EVectorVMOp::logic_or:                             VVMOptimizeVecIns2;             break;
+			case EVectorVMOp::logic_xor:                            VVMOptimizeVecIns2;             break;
+			case EVectorVMOp::logic_not:                            VVMOptimizeVecIns1;             break;
+			case EVectorVMOp::f2i:                                  VVMOptimizeVecIns1;             break;
+			case EVectorVMOp::i2f:                                  VVMOptimizeVecIns1;             break;
+			case EVectorVMOp::f2b:                                  VVMOptimizeVecIns1;             break;
+			case EVectorVMOp::b2f:                                  VVMOptimizeVecIns1;             break;
+			case EVectorVMOp::i2b:                                  VVMOptimizeVecIns1;             break;
+			case EVectorVMOp::b2i:                                  VVMOptimizeVecIns1;             break;
+			case EVectorVMOp::inputdata_float:
+			case EVectorVMOp::inputdata_int32:
+			case EVectorVMOp::inputdata_half:
+			case EVectorVMOp::inputdata_noadvance_float:
+			case EVectorVMOp::inputdata_noadvance_int32:
+			case EVectorVMOp::inputdata_noadvance_half:
+			{
+				uint16 DataSetIdx	= *(uint16 *)(OpPtrIn    );
+				uint16 InputRegIdx	= *(uint16 *)(OpPtrIn + 2);
+				uint16 DstRegIdx	= *(uint16 *)(OpPtrIn + 4);
+				
+				AllocRegisterUse(1);
+				Instruction->Input.DataSetIdx                                                           = DataSetIdx;
+				Instruction->Input.InputIdx                                                             = InputRegIdx;
+				Instruction->Input.DstRegPtrOffset                                                      = OptContext->Intermediate.NumRegistersUsed;
+				OptContext->Intermediate.RegisterUsageBuffer[OptContext->Intermediate.NumRegistersUsed] = DstRegIdx;
+				Instruction->Input.FuseCount                                                            = 0;
+				Instruction->Input.FirstInsInsertIdx                                                    = Instruction->Index;
+
+				++OptContext->Intermediate.NumRegistersUsed;
+				OpPtrIn += 6;
+			}
+			break;
+			case EVectorVMOp::outputdata_float:
+			case EVectorVMOp::outputdata_int32:
+			case EVectorVMOp::outputdata_half:
+			{
+				uint8 OpType        = *OpPtrIn & 1; //0: reg, 1: const
+				uint16 DataSetIdx   = VecIndices[0];
+				uint16 DstIdxRegIdx = VecIndices[1];
+				uint16 SrcReg       = VecIndices[2];
+				uint16 DstRegIdx    = VecIndices[3];
+				check(DataSetIdx < 0xFF);
+				OptContext->NumOutputDataSets = VVM_MAX(OptContext->NumOutputDataSets, (uint32)(DataSetIdx + 1));
+				if (OpType != 0)
+				{
+					SrcReg = VectorVMOptimizeRemapConst(OptContext, SrcReg >> 2) | 0x8000;
+					if (OptContext->Error.Flags & VVMOptErr_Fatal)
+					{
+						return OptContext->Error.Flags;
+					}
+				}
+
+				AllocRegisterUse(2)
+				OptContext->Intermediate.RegisterUsageBuffer[OptContext->Intermediate.NumRegistersUsed + 0] = DstIdxRegIdx;
+				OptContext->Intermediate.RegisterUsageBuffer[OptContext->Intermediate.NumRegistersUsed + 1] = SrcReg;
+				Instruction->Output.DataSetIdx          = DataSetIdx;
+				Instruction->Output.RegPtrOffset        = OptContext->Intermediate.NumRegistersUsed;
+				Instruction->Output.DstRegIdx           = DstRegIdx;
+				Instruction->Output.CopyFromInputInsIdx = -1;
+				OptContext->Intermediate.NumRegistersUsed += 2;
+				OpPtrIn += 9;
+			}
+			break;
+			case EVectorVMOp::acquireindex:
+			{
+				uint8 OpType       = *OpPtrIn & 1;							//0: reg, 1: const4
+				uint16 IncMask     = 0;
+				uint16 DataSetIdx  = VecIndices[0];
+				uint16 OutputReg   = VecIndices[2];
+				uint16 InputRegIdx = 0;
+				if (OpType == 0) { //input register
+					InputRegIdx = VecIndices[1];
+					IncMask = 0xFFFF;
+				} else if (OpType == 1) { //input constant
+					uint16 RemappedIdx = VectorVMOptimizeRemapConst(OptContext, VecIndices[1] >> 2);
+					if (OptContext->Error.Flags & VVMOptErr_Fatal) {
+						return OptContext->Error.Flags;
+					}
+					InputRegIdx = (1 << 15) + RemappedIdx;
+				}
+
+				AllocRegisterUse(3);
+				Instruction->IndexGen.DataSetIdx   = DataSetIdx;
+				Instruction->IndexGen.RegPtrOffset = OptContext->Intermediate.NumRegistersUsed;
+				OptContext->Intermediate.RegisterUsageBuffer[OptContext->Intermediate.NumRegistersUsed + 0] = InputRegIdx;
+				OptContext->Intermediate.RegisterUsageBuffer[OptContext->Intermediate.NumRegistersUsed + 1] = OutputReg;
+				//OptContext->Intermediate.RegisterUsageBuffer[OptContext->Intermediate.NumRegistersUsed + 2] = 0xFFFF;
+				OptContext->Intermediate.NumRegistersUsed += 2;
+				OpPtrIn += 7;
+			}
+			break;
+			case EVectorVMOp::external_func_call:
+			{
+				uint8 ExtFnIdx = *OpPtrIn;
+				check(ExtFnIdx < NumExtFns);
+
+				Instruction->ExtFnCall.RegPtrOffset = OptContext->Intermediate.NumRegistersUsed;
+				Instruction->ExtFnCall.ExtFnIdx     = ExtFnIdx;
+				Instruction->ExtFnCall.NumInputs    = ExtFnIOData[ExtFnIdx].NumInputs;
+				Instruction->ExtFnCall.NumOutputs   = ExtFnIOData[ExtFnIdx].NumOutputs;
+				
+				AllocRegisterUse(ExtFnIOData[ExtFnIdx].NumInputs + ExtFnIOData[ExtFnIdx].NumOutputs);
+				for (int i = 0; i < ExtFnIOData[ExtFnIdx].NumInputs; ++i)
+				{
+					if (VecIndices[i] == 0xFFFF)
+					{ //invalid, just write it out
+						OptContext->Intermediate.RegisterUsageBuffer[OptContext->Intermediate.NumRegistersUsed] = 0xFFFF;
+					}
+					else
+					{
+						if (VecIndices[i] & 0x8000) //register: high bit means input is a register.. the complete opposite behavior of everywhere else.
+						{ 
+							uint16 TempRegIdx = VecIndices[i] & 0x7FFF;
+							OptContext->Intermediate.RegisterUsageBuffer[OptContext->Intermediate.NumRegistersUsed] = TempRegIdx;
+						}
+						else //constant
+						{ 
+							uint16 RemappedIdx = VectorVMOptimizeRemapConst(OptContext, (VecIndices[i] & 0x7FFF) >> 2) | 0x8000; //set the constant flag
+							if (OptContext->Error.Flags & VVMOptErr_Fatal)
+							{
+								return OptContext->Error.Flags;																								\
+							}
+							OptContext->Intermediate.RegisterUsageBuffer[OptContext->Intermediate.NumRegistersUsed] = RemappedIdx;
+						}
+					}
+					++OptContext->Intermediate.NumRegistersUsed;
+				}
+				for (int i = 0; i < ExtFnIOData[ExtFnIdx].NumOutputs; ++i)
+				{
+					int Idx = ExtFnIOData[ExtFnIdx].NumInputs + i;
+					check((VecIndices[Idx] & 0x8000) == 0 || VecIndices[Idx] == 0xFFFF); //can't output to a const... 0xFFFF is invalid
+					OptContext->Intermediate.RegisterUsageBuffer[OptContext->Intermediate.NumRegistersUsed] = VecIndices[Idx];
+					++OptContext->Intermediate.NumRegistersUsed;
+				}
+				OptContext->MaxExtFnUsed      = VVM_MAX(OptContext->MaxExtFnUsed, ExtFnIdx);
+				OptContext->MaxExtFnRegisters = VVM_MAX(OptContext->MaxExtFnRegisters, (uint32)(ExtFnIOData[ExtFnIdx].NumInputs + ExtFnIOData[ExtFnIdx].NumOutputs));
+				OpPtrIn += 1 + (ExtFnIOData[ExtFnIdx].NumInputs + ExtFnIOData[ExtFnIdx].NumOutputs) * 2;
+			}
+			break;
+			case EVectorVMOp::exec_index:
+				AllocRegisterUse(1);
+				OptContext->Intermediate.RegisterUsageBuffer[OptContext->Intermediate.NumRegistersUsed] = *OpPtrIn;
+				Instruction->ExecIndex.RegPtrOffset = OptContext->Intermediate.NumRegistersUsed;
+				++OptContext->Intermediate.NumRegistersUsed;
+				OpPtrIn += 2;
+				break;
+			case EVectorVMOp::noise2D:                              check(false);               break;
+			case EVectorVMOp::noise3D:                              check(false);               break;
+			case EVectorVMOp::enter_stat_scope:						
+				Instruction->Stat.ID = *OpPtrIn;                    OpPtrIn += 2;               break;
+			case EVectorVMOp::exit_stat_scope:                                                  break;
+			case EVectorVMOp::update_id:                            //intentional fallthrough
+			case EVectorVMOp::acquire_id:
+			{
+				uint16 DataSetIdx = ((uint16 *)OpPtrIn)[0];
+				uint16 IDIdxReg = ((uint16 *)OpPtrIn)[1];
+				uint16 IDTagReg = ((uint16 *)OpPtrIn)[2];
+				
+				AllocRegisterUse(2);
+				Instruction->RWBuffer.DataSetIdx   = DataSetIdx;
+				Instruction->RWBuffer.RegPtrOffset = OptContext->Intermediate.NumRegistersUsed;
+				OptContext->Intermediate.RegisterUsageBuffer[OptContext->Intermediate.NumRegistersUsed + 0] = IDIdxReg;
+				OptContext->Intermediate.RegisterUsageBuffer[OptContext->Intermediate.NumRegistersUsed + 1] = IDTagReg;
+				OptContext->Intermediate.NumRegistersUsed += 2;
+				MaxRWBufferUsed = VVM_MAX(MaxRWBufferUsed, DataSetIdx);
+				OpPtrIn += 6;
+			} break;
+			default:												check(false);			     break;
+		}
+	}
+
+	//Step 2: Setup External Function Table
+	if (NumExtFns > 0)
+	{
+		OptContext->ExtFnTable = (FVectorVMExtFunctionData *)OptContext->Init.ReallocFn(nullptr, sizeof(FVectorVMExtFunctionData) * NumExtFns, __FILE__, __LINE__);
+		if (OptContext->ExtFnTable == nullptr)
+		{
+			return VectorVMOptimizerSetError(OptContext, VVMOptErr_OutOfMemory | VVMOptErr_ExternalFunction | VVMOptErr_Fatal);
+		}
+		for (int i = 0; i < NumExtFns; ++i)
+		{
+			OptContext->ExtFnTable[i].NumInputs  = ExtFnIOData[i].NumInputs;
+			OptContext->ExtFnTable[i].NumOutputs = ExtFnIOData[i].NumOutputs;
+		}
+		OptContext->NumExtFns         = NumExtFns;
+	}
+	else
+	{
+		OptContext->ExtFnTable        = nullptr;
+		OptContext->NumExtFns         = 0;
+		OptContext->MaxExtFnRegisters = 0;
+	}
+		
+	{ //Step 3: Verify everything is good
+		{ //verify integirty of constant remap table
+			if (OptContext->NumConstsRemapped > OptContext->NumConstsAlloced)
+			{
+				return VectorVMOptimizerSetError(OptContext, VVMOptErr_OutOfMemory | VVMOptErr_ConstRemap | VVMOptErr_Fatal);
+			}
+			if (OptContext->NumConstsRemapped < OptContext->NumConstsAlloced && OptContext->ConstRemap[1][OptContext->NumConstsRemapped] != 0xFFFF)
+			{
+				return VectorVMOptimizerSetError(OptContext, VVMOptErr_OutOfMemory | VVMOptErr_ConstRemap | VVMOptErr_Fatal);
+			}
+			for (int i = 0; i < OptContext->NumConstsRemapped; ++i)
+			{
+				if (OptContext->ConstRemap[1][i] >= OptContext->NumConstsAlloced)
+				{
+					return VectorVMOptimizerSetError(OptContext, VVMOptErr_OutOfMemory | VVMOptErr_ConstRemap | VVMOptErr_Fatal);
+				}
+				if (OptContext->ConstRemap[0][OptContext->ConstRemap[1][i]] != i)
+				{
+					return VectorVMOptimizerSetError(OptContext, VVMOptErr_OutOfMemory | VVMOptErr_ConstRemap | VVMOptErr_Fatal);
+				}
+			}
+		}
+		if (OptContext->Intermediate.NumRegistersUsed >= 0xFFFF) //16 bit indices
+		{
+			return VectorVMOptimizerSetError(OptContext, VVMOptErr_OutOfMemory | VVMOptErr_RegisterUsage | VVMOptErr_Fatal);
+		}
+	}
+	
+	{ //Step 4: Setup additional buffers
+		{ //setup SSA register buffer
+			OptContext->Intermediate.SSARegisterUsageBuffer = (uint16 *)OptContext->Init.ReallocFn(nullptr, sizeof(uint16) * OptContext->Intermediate.NumRegistersUsed, __FILE__, __LINE__); //upper bound, it can't possibly take more than this
+			if (OptContext->Intermediate.SSARegisterUsageBuffer == nullptr)
+			{
+				return VectorVMOptimizerSetError(OptContext, VVMOptErr_OutOfMemory | VVMOptErr_InputFuseBuffer | VVMOptErr_Fatal);
+			}
+			FMemory::Memcpy(OptContext->Intermediate.SSARegisterUsageBuffer, OptContext->Intermediate.RegisterUsageBuffer, sizeof(uint16) * OptContext->Intermediate.NumRegistersUsed);
+		}
+
+		{ //Setup input fuse buffer
+			OptContext->Intermediate.InputRegisterFuseBuffer = (int32 *)OptContext->Init.ReallocFn(nullptr, sizeof(int32) * OptContext->Intermediate.NumRegistersUsed, __FILE__, __LINE__);
+			if (OptContext->Intermediate.InputRegisterFuseBuffer == nullptr)
+			{
+				return VectorVMOptimizerSetError(OptContext, VVMOptErr_OutOfMemory | VVMOptErr_InputFuseBuffer | VVMOptErr_Fatal);
+			}
+			for (uint32 i = 0; i < OptContext->Intermediate.NumRegistersUsed; ++i)
+			{
+				OptContext->Intermediate.InputRegisterFuseBuffer[i] = -1;
+			}
+		}
+	}
+
+	int OnePastLastInputIdx = -1;
+	uint16 NumSSARegistersUsed = 0;
+
+	{ //Step 5: SSA-like renaming of temp registers
+		FVectorVMOptimizeInsRegUsage InputInsRegUse;
+		FVectorVMOptimizeInsRegUsage OutputInsRegUse;
+				
+		uint16 SSARegCount = 0;
+		for (uint32 OutputInsIdx = 0; OutputInsIdx < OptContext->Intermediate.NumInstructions; ++OutputInsIdx)
+		{
+			FVectorVMOptimizeInstruction *OutputIns = OptContext->Intermediate.Instructions + OutputInsIdx;
+			GetRegistersUsedForInstruction(OptContext, OutputIns, &OutputInsRegUse);
+
+			//loop over each instruction's output
+			for (int j = 0; j < OutputInsRegUse.NumOutputRegisters; ++j)
+			{
+				uint16 OutReg = OptContext->Intermediate.RegisterUsageBuffer[OutputInsRegUse.RegIndices[OutputInsRegUse.NumInputRegisters + j]];
+				if (OutReg != 0xFFFF) {
+					OptContext->Intermediate.SSARegisterUsageBuffer[OutputInsRegUse.RegIndices[OutputInsRegUse.NumInputRegisters + j]] = SSARegCount;
+					int LastUsedAsInputInsIdx = -1;
+					//check each instruction's output with the input of every instruction that follows it
+					for (uint32 InputInsIdx = OutputInsIdx + 1; InputInsIdx < OptContext->Intermediate.NumInstructions; ++InputInsIdx)
+					{
+						FVectorVMOptimizeInstruction *InputIns = OptContext->Intermediate.Instructions + InputInsIdx;
+						GetRegistersUsedForInstruction(OptContext, InputIns, &InputInsRegUse);
+					
+						//check to see if the register we're currently looking at (OutReg) is overwritten by another instruction.  If it is,
+						//we increment the SSA count, and move on to the next 
+						for (int ii = 0; ii < InputInsRegUse.NumOutputRegisters; ++ii)
+						{
+							if (OptContext->Intermediate.RegisterUsageBuffer[InputInsRegUse.RegIndices[InputInsRegUse.NumInputRegisters + ii]] == OutReg)
+							{
+								//this register is overwritten, we need to generate a new register
+								++SSARegCount;
+								check(SSARegCount <= OptContext->Intermediate.NumRegistersUsed);
+								goto DoneThisOutput;
+							}
+						}
+
+						//if the Input instruction's input uses the Output instruction's output then assign them to the same SSA value
+						for (int ii = 0; ii < InputInsRegUse.NumInputRegisters; ++ii)
+						{
+							if (OptContext->Intermediate.RegisterUsageBuffer[InputInsRegUse.RegIndices[ii]] == OutReg)
+							{
+								if (InputIns->OpCat == EVectorVMOpCategory::Output)
+								{
+									if (InputIns->OpCode != EVectorVMOp::copy_to_output)
+									{
+										if (OutputIns->OpCode == EVectorVMOp::acquireindex)
+										{
+											if (j == ii)
+											{
+												//We're only comparing acquireindex's output 0 with the output instructions input 0
+												//or acquireindex's output 1 with the output instructions input 1.
+												//This is because acquireindex's output 0 is the actual index that the output will write to
+												//acquireindex's output 1 is the previous VM's acquireindex's output... this means that if
+												//the output instruction is writing the generated index, it needs to write the *PREVIOUS VM's*
+												//output, since somewhere else in the engine is expecting it... however if the output instruction
+												//is using it as the index to write, we use the new VM since it's optimized for writing.
+												check(j == 0);
+												//if this assert hits it means that the output of the acquireindex is being written to a buffer...
+												//on Dec 6, 2021 I asked about this in the #niagara-cpuvm-optimizations channel on slack and
+												//at the time there was no way to hook up nodes to create this output... I was also informed
+												//it was unlikly that this would ever be possible... so if this assert triggered either this feature
+												//was added, or there's a bug.  It *COULD* work just fine, it's just never been tested.
+												OptContext->Intermediate.SSARegisterUsageBuffer[InputIns->Output.RegPtrOffset + j] = SSARegCount;
+												LastUsedAsInputInsIdx = InputInsIdx;
+											}
+										}
+										else
+										{
+											OptContext->Intermediate.SSARegisterUsageBuffer[InputIns->Output.RegPtrOffset + 1] = SSARegCount;
+											LastUsedAsInputInsIdx = InputInsIdx;
+										}
+									}
+								}
+								else
+								{
+									LastUsedAsInputInsIdx = InputInsIdx;
+									OptContext->Intermediate.SSARegisterUsageBuffer[InputInsRegUse.RegIndices[ii]] = SSARegCount;
+								}
+							}
+						}
+					}
+					if (LastUsedAsInputInsIdx != -1)
+					{
+						++SSARegCount;
+					}
+				} else {
+					OptContext->Intermediate.SSARegisterUsageBuffer[OutputInsRegUse.RegIndices[OutputInsRegUse.NumInputRegisters + j]] = 0xFFFF;
+				}
+				DoneThisOutput: ;
+			}
+		}
+		check(SSARegCount < 0xFFFF - 1);
+		NumSSARegistersUsed = SSARegCount + 1;
+	}
+
+	
+	{ //Step 6: Input Fusing
+		//gather all input instructions can be fused
+		for (uint32 i = 0; i < OptContext->Intermediate.NumInstructions; ++i)
+		{
+			FVectorVMOptimizeInstruction *InputIns = OptContext->Intermediate.Instructions + i;
+			if (InputIns->OpCode == EVectorVMOp::inputdata_float || 
+				InputIns->OpCode == EVectorVMOp::inputdata_int32 || 
+				InputIns->OpCode == EVectorVMOp::inputdata_half) //noadvance ops can't fuse
+			{ 
+				OnePastLastInputIdx = i + 1;
+				bool InputOpCanFuse = true;
+				InputIns->Input.FirstInsInsertIdx = -1;
+				for (uint32 j = i + 1; j < OptContext->Intermediate.NumInstructions; ++j)
+				{
+					FVectorVMOptimizeInstruction *Ins_j = OptContext->Intermediate.Instructions + j;
+					switch (Ins_j->OpCat)
+					{
+						case EVectorVMOpCategory::Input: //make sure there isn't an instruction that overwrites this input.  If this happens then this is a useless input
+							if (OptContext->Intermediate.SSARegisterUsageBuffer[InputIns->Input.DstRegPtrOffset] == OptContext->Intermediate.SSARegisterUsageBuffer[Ins_j->Input.DstRegPtrOffset])
+							{
+								check(false);
+							}
+							break;
+						case EVectorVMOpCategory::Output:
+							if (OptContext->Intermediate.SSARegisterUsageBuffer[InputIns->Input.DstRegPtrOffset] == OptContext->Intermediate.SSARegisterUsageBuffer[Ins_j->Output.RegPtrOffset + 1])
+							{
+								++InputIns->Input.FuseCount;
+								Ins_j->Output.CopyFromInputInsIdx = i;
+							}
+							break;
+						case EVectorVMOpCategory::Op: {
+							uint16 *Registers = OptContext->Intermediate.SSARegisterUsageBuffer + Ins_j->Op.RegPtrOffset;
+							//check to see if this matches an input
+							int NumInputsToCheck = Ins_j->Op.NumInputs;
+							if (NumInputsToCheck > 3)
+							{
+								NumInputsToCheck = 3;
+							}
+							for (int k = 0; k < NumInputsToCheck; ++k)
+							{
+								if (Registers[k] == OptContext->Intermediate.SSARegisterUsageBuffer[InputIns->Input.DstRegPtrOffset])
+								{
+									//This input and operation are fusable
+									OptContext->Intermediate.InputRegisterFuseBuffer[Ins_j->Op.RegPtrOffset + k] = i;
+									Ins_j->Op.InputFuseBits |= (1 << k);
+									++InputIns->Input.FuseCount;
+								}
+							}
+							if (Ins_j->Op.NumInputs > 3) //we can only fuse the first 3 inputs, after that we need to have an explicit input instruction
+							{
+								for (int k = 3; k < Ins_j->Op.NumInputs; ++k)
+								{
+									if (Registers[k] == OptContext->Intermediate.SSARegisterUsageBuffer[InputIns->Input.DstRegPtrOffset])
+									{
+										InputIns->Input.FirstInsInsertIdx = InputIns->Input.FirstInsInsertIdx == -1 ? ((int)j - 1) : VVM_MIN(InputIns->Input.FirstInsInsertIdx, (int)j - 1);
+									}
+								}
+							}
+							for (int k = 0; k < Ins_j->Op.NumOutputs; ++k)
+							{
+								if (Registers[Ins_j->Op.NumInputs + k] == OptContext->Intermediate.SSARegisterUsageBuffer[InputIns->Input.DstRegPtrOffset])
+								{
+									//this register is used as an output... therefore it's no longer a canditate for input fusing and we need to stop checking
+									InputOpCanFuse = false;
+								}
+							}
+						}	break;
+						case EVectorVMOpCategory::ExtFnCall: 
+							//we don't allow input op fusing on external functions (yet? may be worth doing maybe? we'd have to change the FVectorVMExternalFunctionContext::DecodeNextRegister function to check for fused inputs)
+							//we do however need to check to see if the function uses the potential register as an output... in this case we need to stop because we can no longer fuse the input
+							//as the register the input copies to is no longer valid
+							for (int k = 0; k < Ins_j->ExtFnCall.NumInputs; ++k)
+							{
+								if (OptContext->Intermediate.SSARegisterUsageBuffer[Ins_j->Op.RegPtrOffset + k] == OptContext->Intermediate.SSARegisterUsageBuffer[InputIns->Input.DstRegPtrOffset])
+								{
+									InputIns->Input.FirstInsInsertIdx = InputIns->Input.FirstInsInsertIdx == -1 ? ((int)j - 1) : VVM_MIN(InputIns->Input.FirstInsInsertIdx, (int)j - 1);
+								}
+							}
+							for (int k = 0; k < Ins_j->ExtFnCall.NumOutputs; ++k)
+							{
+								if (OptContext->Intermediate.SSARegisterUsageBuffer[Ins_j->Op.RegPtrOffset + Ins_j->ExtFnCall.NumInputs + k] == OptContext->Intermediate.SSARegisterUsageBuffer[InputIns->Input.DstRegPtrOffset])
+								{
+									InputOpCanFuse = false;
+									break;
+								}
+							}
+							break;
+						case EVectorVMOpCategory::IndexGen: //can't fuse to index gen
+							if (OptContext->Intermediate.SSARegisterUsageBuffer[InputIns->Input.DstRegPtrOffset] == OptContext->Intermediate.SSARegisterUsageBuffer[Ins_j->IndexGen.RegPtrOffset])
+							{
+								InputIns->Input.FirstInsInsertIdx = InputIns->Input.FirstInsInsertIdx == -1 ? ((int)j - 1) : VVM_MIN(InputIns->Input.FirstInsInsertIdx, (int)j - 1);
+							} 
+							else if (OptContext->Intermediate.SSARegisterUsageBuffer[InputIns->Input.DstRegPtrOffset] == OptContext->Intermediate.SSARegisterUsageBuffer[Ins_j->IndexGen.RegPtrOffset + 1])
+							{
+								InputOpCanFuse = false;
+							}
+							break;
+						case EVectorVMOpCategory::ExecIndex:
+							if (OptContext->Intermediate.SSARegisterUsageBuffer[InputIns->Input.DstRegPtrOffset] == OptContext->Intermediate.SSARegisterUsageBuffer[Ins_j->ExecIndex.RegPtrOffset])
+							{
+								InputOpCanFuse = false; //exec_index is output only, so this register is being overwritten.
+							}
+							break;
+						case EVectorVMOpCategory::RWBuffer:
+							if (Ins_j->OpCode == EVectorVMOp::update_id) //update_id is input only, acquire_id is output only
+							{
+								if (OptContext->Intermediate.SSARegisterUsageBuffer[InputIns->Input.DstRegPtrOffset] == OptContext->Intermediate.SSARegisterUsageBuffer[Ins_j->RWBuffer.RegPtrOffset + 0] ||
+									OptContext->Intermediate.SSARegisterUsageBuffer[InputIns->Input.DstRegPtrOffset] == OptContext->Intermediate.SSARegisterUsageBuffer[Ins_j->RWBuffer.RegPtrOffset + 1])
+								{
+									InputIns->Input.FirstInsInsertIdx = InputIns->Input.FirstInsInsertIdx == -1 ? ((int)j - 1) : VVM_MIN(InputIns->Input.FirstInsInsertIdx, (int)j - 1); //no fusing inputs to update_id... for no particular reason.  It's just more complicated and not really necessary.  If we find it's useful I can add it
+								}
+							}
+							break;
+						case EVectorVMOpCategory::Stat:
+							break;
+						case EVectorVMOpCategory::Other:
+							check(Ins_j->OpCode == EVectorVMOp::done || Ins_j->OpCode == EVectorVMOp::noise2D || Ins_j->OpCode == EVectorVMOp::noise3D);
+							break;
+					}
+					if (!InputOpCanFuse)
+					{
+						break;
+					}
+				}
+			}
+		}
+		//skip the stat instructions after the inputs
+		if (OnePastLastInputIdx == -1)
+		{
+			OnePastLastInputIdx = 0;
+		}
+		while (OnePastLastInputIdx != -1 && OptContext->Intermediate.NumInstructions != 0 && (uint32)OnePastLastInputIdx < OptContext->Intermediate.NumInstructions - 1 && OptContext->Intermediate.Instructions[OnePastLastInputIdx].OpCat == EVectorVMOpCategory::Stat)
+		{
+			++OnePastLastInputIdx;
+		}
+	}
+
+	if (0)
+	{ //Step 7: remove instructions where outputs are never used 
+		int NumRemovedInstructions = 0;
+		FVectorVMOptimizeInsRegUsage RegUsage;
+		FVectorVMOptimizeInsRegUsage RegUsage2;
+		int NumRemovedInstructionsThisTime;
+		int SanityCount = 0;
+		do
+		{
+			//loop multiple times because sometimes an instruction can be removed that will make a previous instruction redundant as well
+			NumRemovedInstructionsThisTime = 0;
+			for (uint32 i = 0; i < OptContext->Intermediate.NumInstructions; ++i)
+			{
+				FVectorVMOptimizeInstruction *Ins = OptContext->Intermediate.Instructions + i;
+				if (Ins->OpCat == EVectorVMOpCategory::Op && !(Ins->OpCode == EVectorVMOp::random || Ins->OpCode == EVectorVMOp::randomi))
+				{
+					//can we remove random instructions? I dunno! so lets not for now
+					bool InsRequired = false;
+					GetRegistersUsedForInstruction(OptContext, Ins, &RegUsage);
+					for (int OutputIdx = 0; OutputIdx < RegUsage.NumOutputRegisters; ++OutputIdx)
+					{
+						uint16 RegIdx = OptContext->Intermediate.SSARegisterUsageBuffer[RegUsage.RegIndices[RegUsage.NumInputRegisters + OutputIdx]];
+						for (uint32 j = i + 1; j < OptContext->Intermediate.NumInstructions; ++j)
+						{
+							GetRegistersUsedForInstruction(OptContext, OptContext->Intermediate.Instructions + j, &RegUsage2);
+							for (int k = 0; k < RegUsage2.NumInputRegisters; ++k)
+							{
+								uint16 RegIdx2 = OptContext->Intermediate.SSARegisterUsageBuffer[RegUsage2.RegIndices[k]];
+								if (RegIdx == RegIdx2)
+								{
+									InsRequired = true;
+									break;
+								}
+							}
+						}
+					}
+					if (!InsRequired)
+					{
+						FMemory::Memmove(OptContext->Intermediate.Instructions + i, OptContext->Intermediate.Instructions + i + 1, sizeof(FVectorVMOptimizeInstruction) * (OptContext->Intermediate.NumInstructions - i - 1));
+						++NumRemovedInstructionsThisTime;
+						++NumRemovedInstructions;
+						--OptContext->Intermediate.NumInstructions;
+						--i;
+					}
+				}
+			}
+			if (++SanityCount >= 16384)
+			{
+				check(false);
+				NumRemovedInstructions = 0;
+				break;
+			}
+		} while (NumRemovedInstructionsThisTime > 0);
+		
+		if (NumRemovedInstructions > 0) //Step 8: re-assign SSA registers if we removed instructions
+		{
+			FMemory::Memcpy(OptContext->Intermediate.SSARegisterUsageBuffer, OptContext->Intermediate.RegisterUsageBuffer, sizeof(uint16) * OptContext->Intermediate.NumRegistersUsed);
+			NumSSARegistersUsed = 0;
+
+			FVectorVMOptimizeInsRegUsage InputInsRegUse;
+			FVectorVMOptimizeInsRegUsage OutputInsRegUse;
+				
+			int SSARegCount = 0;
+			for (uint32 i = 0; i < OptContext->Intermediate.NumInstructions; ++i)
+			{
+				FVectorVMOptimizeInstruction *OutputIns = OptContext->Intermediate.Instructions + i;
+				GetRegistersUsedForInstruction(OptContext, OutputIns, &OutputInsRegUse);
+				if (OutputInsRegUse.NumOutputRegisters > 0)
+				{
+					for (int j = 0; j < OutputInsRegUse.NumOutputRegisters; ++j)
+					{
+						uint16 OutReg = OptContext->Intermediate.RegisterUsageBuffer[OutputInsRegUse.RegIndices[OutputInsRegUse.NumInputRegisters + j]];
+						OptContext->Intermediate.SSARegisterUsageBuffer[OutputInsRegUse.RegIndices[OutputInsRegUse.NumInputRegisters + j]] = SSARegCount;
+						int LastUsedAsInputInsIdx = -1;
+						for (uint32 k = i + 1; k < OptContext->Intermediate.NumInstructions; ++k)
+						{
+							FVectorVMOptimizeInstruction *InputIns = OptContext->Intermediate.Instructions + k;
+							GetRegistersUsedForInstruction(OptContext, InputIns, &InputInsRegUse);
+							for (int ii = 0; ii < InputInsRegUse.NumOutputRegisters; ++ii)
+							{
+								if (OptContext->Intermediate.RegisterUsageBuffer[InputInsRegUse.RegIndices[InputInsRegUse.NumInputRegisters + ii]] == OutReg)
+								{
+									//this register is overwritten, we need to generate a new register
+									++SSARegCount;
+									check(SSARegCount <= (int)OptContext->Intermediate.NumRegistersUsed);
+									goto DoneThisOutput2;
+								}
+							}
+							for (int ii = 0; ii < InputInsRegUse.NumInputRegisters; ++ii)
+							{
+								if (OptContext->Intermediate.RegisterUsageBuffer[InputInsRegUse.RegIndices[ii]] == OutReg)
+								{
+									if (InputIns->OpCat == EVectorVMOpCategory::Output)
+									{
+										if (InputIns->OpCode != EVectorVMOp::copy_to_output)
+										{
+											if (OutputIns->OpCode == EVectorVMOp::acquireindex)
+											{
+												OptContext->Intermediate.SSARegisterUsageBuffer[InputIns->Output.RegPtrOffset] = SSARegCount;
+											}
+											else
+											{
+												OptContext->Intermediate.SSARegisterUsageBuffer[InputIns->Output.RegPtrOffset + 1] = SSARegCount;
+												LastUsedAsInputInsIdx = k;
+											}
+										}
+									}
+									else
+									{
+										LastUsedAsInputInsIdx = k;
+										OptContext->Intermediate.SSARegisterUsageBuffer[InputInsRegUse.RegIndices[ii]] = SSARegCount;
+									}
+								}
+							}
+						}
+						if (LastUsedAsInputInsIdx != -1)
+						{
+							++SSARegCount;
+						}
+						DoneThisOutput2: ;
+					}
+				}
+			}
+			check(SSARegCount < 0xFFFF - 1);
+			NumSSARegistersUsed = SSARegCount + 1;
+		}
+	}
+
+	if (1) { //instruction re-ordering
+		int *RegToCheckStack = (int *)OptContext->Init.ReallocFn(nullptr, sizeof(int) * OptContext->Intermediate.NumRegistersUsed * 2, __FILE__, __LINE__);	//these two could actually be a single array, 1/2 the size, one starting from 0 and counting up, the other one
+		if (RegToCheckStack == nullptr)
+		{
+			return VectorVMOptimizerSetError(OptContext, VVMOptErr_OutOfMemory | VVMOptErr_InstructionReOrder | VVMOptErr_Fatal);
+		}
+		VVMOptRAIIPtrToFree RegStackRAII(OptContext, RegToCheckStack);
+
+		int *InstructionIdxStack = RegToCheckStack + OptContext->Intermediate.NumRegistersUsed;
+		int LowestInstructionIdxForAcquireIdx = OnePastLastInputIdx; //acquire index instructions will be sorted by whichever comes first in the IR... possibly worth checking if re-ordering is more efficient
+
+		{ //Step 9: Find all the acquireindex instructions and re-order them to be executed ASAP
+			int NumAcquireIndexInstructions = 0;
+			for (uint32 i = 0; i < OptContext->Intermediate.NumInstructions; ++i)
+			{
+				FVectorVMOptimizeInstruction *Ins = OptContext->Intermediate.Instructions + i;
+				if (Ins->OpCode == EVectorVMOp::acquireindex)
+				{
+					++NumAcquireIndexInstructions;
+					int AcquireIndexInstructionIdx = i;
+					int NumInstructions = GetInstructionDependencyChain(OptContext, AcquireIndexInstructionIdx, RegToCheckStack, InstructionIdxStack);
+					//bubble up the dependent instructions at quickly as possible (@NOTE: if these are already grouped together it's more efficient to move more than 1 instruction at a time, but who cares? I doubt the performace will ever matter)
+					for (int j = 0; j < NumInstructions; ++j)
+					{
+						if (InstructionIdxStack[j] > LowestInstructionIdxForAcquireIdx)
+						{
+							FVectorVMOptimizeInstruction TempIns = OptContext->Intermediate.Instructions[InstructionIdxStack[j]];
+							FMemory::Memmove(OptContext->Intermediate.Instructions + LowestInstructionIdxForAcquireIdx + 1, OptContext->Intermediate.Instructions + LowestInstructionIdxForAcquireIdx, sizeof(FVectorVMOptimizeInstruction) * (InstructionIdxStack[j] - LowestInstructionIdxForAcquireIdx));
+							OptContext->Intermediate.Instructions[LowestInstructionIdxForAcquireIdx] = TempIns;
+						}
+						LowestInstructionIdxForAcquireIdx = LowestInstructionIdxForAcquireIdx + 1;
+					}
+					//move the acquire index instruction to immediately after the last instruction it depends on
+					if (LowestInstructionIdxForAcquireIdx < AcquireIndexInstructionIdx)
+					{
+						FVectorVMOptimizeInstruction TempIns = OptContext->Intermediate.Instructions[AcquireIndexInstructionIdx];
+						FMemory::Memmove(OptContext->Intermediate.Instructions + LowestInstructionIdxForAcquireIdx + 1, OptContext->Intermediate.Instructions + LowestInstructionIdxForAcquireIdx, sizeof(FVectorVMOptimizeInstruction) * (AcquireIndexInstructionIdx - LowestInstructionIdxForAcquireIdx));
+						OptContext->Intermediate.Instructions[LowestInstructionIdxForAcquireIdx++] = TempIns;
+					}
+				}
+			}
+			//there's a potential race condition if two acquireindex instructions are in a script and they're run multithreaded... ie:
+			//threadA: acquireindex0
+			//threadB: acquireindex0
+			//threadB: acquireindex1
+			//threadA: acquireindex1
+			//in this situation the output data will not match between two datasets.. ie: instance0 in dataset0 will not be correlated to instance0 in dataset1.  If they were running single threaded that wouldn't happen.
+		}
+
+		if (1) { //Step 10: Find all update_id instructions and re-order them to be just after their inputs
+			FVectorVMOptimizeInsRegUsage RegUsage;
+			for (uint32 i = 0; i < OptContext->Intermediate.NumInstructions; ++i)
+			{
+				FVectorVMOptimizeInstruction *Ins = OptContext->Intermediate.Instructions + i;
+				if (Ins->OpCode == EVectorVMOp::update_id)
+				{
+					uint32 InsertionIdx = 0xFFFFFFFF;
+					GetRegistersUsedForInstruction(OptContext, Ins, &RegUsage);
+					check(RegUsage.NumInputRegisters == 2);
+					uint16 UpdateIdxReg[2] = { OptContext->Intermediate.SSARegisterUsageBuffer[RegUsage.RegIndices[0]], OptContext->Intermediate.SSARegisterUsageBuffer[RegUsage.RegIndices[1]] };
+					for (uint32 j = 0; j < i; ++j)
+					{
+						if (OptContext->Intermediate.Instructions[j].OpCode == EVectorVMOp::acquire_id && OptContext->Intermediate.Instructions[j].RWBuffer.DataSetIdx == Ins->RWBuffer.DataSetIdx)
+						{
+							InsertionIdx = j + 1; //update_id must come after the acquire_id for the same DataSet
+						}
+						else
+						{
+							GetRegistersUsedForInstruction(OptContext, OptContext->Intermediate.Instructions + j, &RegUsage);
+							for (int k = 0; k < RegUsage.NumOutputRegisters; ++k)
+							{
+								uint16 RegIdx = OptContext->Intermediate.SSARegisterUsageBuffer[RegUsage.RegIndices[RegUsage.NumInputRegisters + k]];
+								if (RegIdx == UpdateIdxReg[0] || RegIdx == UpdateIdxReg[1])
+								{
+									InsertionIdx = j + 1;
+								}
+							}
+						}
+					}
+					if (InsertionIdx != 0xFFFFFFFF && InsertionIdx + 2 < i)
+					{
+						 FVectorVMOptimizeInstruction TempIns = OptContext->Intermediate.Instructions[i];
+						 FMemory::Memmove(OptContext->Intermediate.Instructions + InsertionIdx + 1,
+							 OptContext->Intermediate.Instructions + InsertionIdx,
+							 sizeof(FVectorVMOptimizeInstruction) * (i - InsertionIdx));
+						 OptContext->Intermediate.Instructions[InsertionIdx] = TempIns;
+					}
+				}
+			}
+		}
+
+		{ //Step 11: re-order the outputs to be done as early as possible: after the SSA's register's last usage
+			for (uint32 OutputInsIdx = 0; OutputInsIdx < OptContext->Intermediate.NumInstructions; ++OutputInsIdx)
+			{
+				FVectorVMOptimizeInstruction *OutputIns = OptContext->Intermediate.Instructions + OutputInsIdx;
+				if (OutputIns->OpCat == EVectorVMOpCategory::Output && OutputIns->Output.CopyFromInputInsIdx == -1)
+				{
+					uint32 OutputInsertionIdx = 0xFFFFFFFF;
+					bool FoundAcquireIndex = false;
+					uint16 IdxReg = OptContext->Intermediate.SSARegisterUsageBuffer[OutputIns->Output.RegPtrOffset];
+					uint16 SrcReg = OptContext->Intermediate.SSARegisterUsageBuffer[OutputIns->Output.RegPtrOffset + 1];
+					for (uint32 i = 0; i < OutputInsIdx; ++i)
+					{
+						FVectorVMOptimizeInstruction *Ins = OptContext->Intermediate.Instructions + i;
+						FVectorVMOptimizeInsRegUsage RegUsage;
+						int NumRegisters = GetRegistersUsedForInstruction(OptContext, Ins, &RegUsage);
+						if (((Ins->OpCat == EVectorVMOpCategory::Input && Ins->Input.FirstInsInsertIdx != -1) || Ins->OpCat != EVectorVMOpCategory::Input) && Ins->OpCat != EVectorVMOpCategory::Output) {
+							for (int j = 0; j < RegUsage.NumOutputRegisters; ++j)
+							{
+								if (OptContext->Intermediate.SSARegisterUsageBuffer[RegUsage.RegIndices[RegUsage.NumInputRegisters + j]] == IdxReg)
+								{
+									FoundAcquireIndex = true;
+									OutputInsertionIdx = i + 1;
+								}
+							}
+							if (FoundAcquireIndex)
+							{
+								for (int j = 0; j < NumRegisters; ++j)
+								{
+									if (OptContext->Intermediate.SSARegisterUsageBuffer[RegUsage.RegIndices[j]] == SrcReg)
+									{
+										OutputInsertionIdx = i + 1;
+									}
+								}
+							}
+						}
+					}
+					if (OutputInsertionIdx != 0xFFFFFFFF && OutputInsertionIdx < OptContext->Intermediate.NumInstructions - 1)
+					{
+						if (OutputInsIdx > OutputInsertionIdx)
+						{
+							uint32 NumInstructionsToMove = OutputInsIdx - OutputInsertionIdx;
+							FVectorVMOptimizeInstruction TempIns = *OutputIns;
+							FMemory::Memmove(OptContext->Intermediate.Instructions + OutputInsertionIdx + 1, OptContext->Intermediate.Instructions + OutputInsertionIdx, sizeof(FVectorVMOptimizeInstruction) * NumInstructionsToMove);
+							OptContext->Intermediate.Instructions[OutputInsertionIdx] = TempIns;
+						}
+					}
+				}
+			}
+		}
+
+		{ //Step 12: re-order all dependent-less instructions to right before their output is used
+			int LastSwapInstructionIdx = -1; //to prevent an infinite loop when one instruction has two or more dependencies and they keep swapping back and forth
+			for (uint32 i = 0; i < OptContext->Intermediate.NumInstructions; ++i)
+			{
+				FVectorVMOptimizeInstruction *Ins = OptContext->Intermediate.Instructions + i;
+				int SkipInstructionSwap = LastSwapInstructionIdx;
+				LastSwapInstructionIdx = -1;
+				if (Ins->OpCat == EVectorVMOpCategory::Op)
+				{
+					int OpNumDependents = GetInstructionDependencyChain(OptContext, i, RegToCheckStack, InstructionIdxStack);
+					if (OpNumDependents == 0)
+					{
+						for (uint32 j = i + 1; j < OptContext->Intermediate.NumInstructions; ++j)
+						{
+							FVectorVMOptimizeInstruction *DepIns = OptContext->Intermediate.Instructions + j;
+							int NumDependents = GetInstructionDependencyChain(OptContext, j, RegToCheckStack, InstructionIdxStack);
+							uint32 InsDepIdx = 0xFFFFFFFF;
+							for (int k = 0; k < NumDependents; ++k)
+							{
+								if (InstructionIdxStack[k] == i)
+								{
+									InsDepIdx = j;
+									break;
+								}
+							}
+							if (InsDepIdx != 0xFFFFFFFF)
+							{
+								if (InsDepIdx > i + 1 && InsDepIdx != SkipInstructionSwap) //DepIns is depdenent on Ins.  Move Ins to be right before DepIns
+								{
+									FVectorVMOptimizeInstruction TempIns = *Ins;
+									FMemory::Memmove(Ins, Ins + 1, sizeof(FVectorVMOptimizeInstruction) * (InsDepIdx - i - 1));
+									OptContext->Intermediate.Instructions[InsDepIdx - 1] = TempIns;
+									LastSwapInstructionIdx = InsDepIdx;
+									--i;
+								}
+								break; //we stop checking even if we don't move the instruction because it's already immediately before its first usage
+							}
+						}
+					}
+				}
+			}
+		}
+
+		if (1) { //Step 13: re-order all inputs to directly before they're used
+			FVectorVMOptimizeInsRegUsage RegUsage;
+			for (uint32 i = 0; i < OptContext->Intermediate.NumInstructions; ++i)
+			{
+				FVectorVMOptimizeInstruction *InputIns = OptContext->Intermediate.Instructions + i;
+				uint16 InputReg = OptContext->Intermediate.SSARegisterUsageBuffer[InputIns->Input.DstRegPtrOffset];
+				if (InputIns->OpCat == EVectorVMOpCategory::Input && InputIns->Input.FirstInsInsertIdx != -1)
+				{
+					for (uint32 j = i + 1; j < OptContext->Intermediate.NumInstructions; ++j)
+					{
+						FVectorVMOptimizeInstruction *OpIns = OptContext->Intermediate.Instructions + j;
+						if (OpIns->OpCat == EVectorVMOpCategory::Output && OpIns->Output.CopyFromInputInsIdx == i)
+						{
+							continue;
+						}
+						GetRegistersUsedForInstruction(OptContext, OpIns, &RegUsage);
+						bool MoveInputHere = false;
+						for (int k = 0; k < RegUsage.NumInputRegisters; ++k)
+						{
+							if (OptContext->Intermediate.SSARegisterUsageBuffer[RegUsage.RegIndices[k]] == InputReg)
+							{
+								MoveInputHere = true;
+								break;
+							}
+						}
+
+						if (MoveInputHere)
+						{
+							if (j > i + 1)
+							{
+								int NewInputIndex = j - 1;
+								FVectorVMOptimizeInstruction TempIns = *InputIns;
+								FMemory::Memmove(InputIns, InputIns + 1, sizeof(FVectorVMOptimizeInstruction) * (j - i - 1));
+								OptContext->Intermediate.Instructions[NewInputIndex] = TempIns;
+								bool ReorderingInputs = true;
+								//if we're only moving this instruction before inputs then we'll get into an infinite loop just re-ordering
+								//inputs around each other.  Check for that and skip all these if that's the case
+								check(i < j);
+								for (uint32 k = i; k < j; ++k)
+								{
+									if (OptContext->Intermediate.Instructions[k].OpCat != EVectorVMOpCategory::Input)
+									{
+										ReorderingInputs = false;
+									}
+								}
+								if (!ReorderingInputs)
+								{
+									--i;
+								}
+							}
+							break;
+						}
+					}
+				}
+			}
+		}
+	}
+
+	{ //Step 14: make sure all the copy-to-output instructions are grouped together
+		int FirstCopyFromInputInsIdx = -1;
+		int LastCopyFromInputInsIdx = -1;
+		//when the copy-to-output instructions get written to the bytecode they'll get grouped into as few instructions as possible
+		for (uint32 i = 0; i < OptContext->Intermediate.NumInstructions; ++i)
+		{
+			if (OptContext->Intermediate.Instructions[i].OpCat == EVectorVMOpCategory::Output && OptContext->Intermediate.Instructions[i].Output.CopyFromInputInsIdx != -1)
+			{
+				if (FirstCopyFromInputInsIdx == -1)
+				{
+					FirstCopyFromInputInsIdx = i;
+				}
+				LastCopyFromInputInsIdx = i;
+			}
+		}
+		//if there's a gap, move the non-copy-to-output instruction to before the copies
+		if (FirstCopyFromInputInsIdx < LastCopyFromInputInsIdx - 1)
+		{
+			for (int i = FirstCopyFromInputInsIdx; i < LastCopyFromInputInsIdx; ++i)
+			{
+				FVectorVMOptimizeInstruction *Ins = OptContext->Intermediate.Instructions + i;
+				if (Ins->OpCat != EVectorVMOpCategory::Output || Ins->Output.CopyFromInputInsIdx == -1)
+				{
+					FVectorVMOptimizeInstruction TempIns = *Ins;
+					FMemory::Memmove(OptContext->Intermediate.Instructions + FirstCopyFromInputInsIdx + 1, OptContext->Intermediate.Instructions + FirstCopyFromInputInsIdx, sizeof(FVectorVMOptimizeInstruction) * (i - FirstCopyFromInputInsIdx));
+					OptContext->Intermediate.Instructions[FirstCopyFromInputInsIdx] = TempIns;
+					++FirstCopyFromInputInsIdx;
+				}
+			}
+		}
+
+		
+		if (LastCopyFromInputInsIdx >= FirstCopyFromInputInsIdx + 2)
+		{
+			//fixup CopyFromInputIns for outpus where the corresponding input instruction was moved
+			for (int i = FirstCopyFromInputInsIdx; i <= LastCopyFromInputInsIdx; ++i)
+			{
+				FVectorVMOptimizeInstruction *OutputIns = OptContext->Intermediate.Instructions + i;
+				check(OutputIns->OpCat == EVectorVMOpCategory::Output);
+				check(OutputIns->Output.CopyFromInputInsIdx != -1);
+				FVectorVMOptimizeInstruction *InputIns = OptContext->Intermediate.Instructions + OutputIns->Output.CopyFromInputInsIdx;
+				if (InputIns->Index == OutputIns->Output.CopyFromInputInsIdx)
+				{
+					check(InputIns->OpCat == EVectorVMOpCategory::Input);
+				}
+				else //this input instruction has been re-ordered.  Fix the output instruction so it points to the correct input instruction
+				{ 
+					for (uint32 j = 0; j < OptContext->Intermediate.NumInstructions; ++j)
+					{
+						FVectorVMOptimizeInstruction *Ins = OptContext->Intermediate.Instructions + j;
+						if (Ins->Index == OutputIns->Output.CopyFromInputInsIdx)
+						{
+							check(Ins->OpCat == EVectorVMOpCategory::Input);
+							check(OutputIns->Output.CopyFromInputInsIdx != (int)j); //should have found it above and never hit this loop if it's already correct
+							OutputIns->Output.CopyFromInputInsIdx = (int)j;
+							break;
+						}
+					}
+				}
+			}
+
+			{ // small list, very stupid bubble sort
+				bool sorted = false;
+				int NumCopyInstructions                        = LastCopyFromInputInsIdx - FirstCopyFromInputInsIdx + 1;
+				while (!sorted) {
+					sorted = true;
+					for (int i = 0; i < NumCopyInstructions - 1; ++i) {
+						FVectorVMOptimizeInstruction *Ins0 = OptContext->Intermediate.Instructions + FirstCopyFromInputInsIdx + i;
+						FVectorVMOptimizeInstruction *Ins1 = Ins0 + 1;
+						uint64 Key0 = VVMCopyToOutputInsGetSortKey(OptContext->Intermediate.Instructions, Ins0);
+						uint64 Key1 = VVMCopyToOutputInsGetSortKey(OptContext->Intermediate.Instructions, Ins1);
+
+						if (Key1 < Key0) {
+							FVectorVMOptimizeInstruction Temp = *Ins0;
+							*Ins0 = *Ins1;
+							*Ins1 = Temp;
+							sorted = false;
+						}
+					}
+				}
+			}
+		}
+	}
+
+	{ //Step 15: group and sort all regular output instructions
+		for (uint32 i = 0; i < OptContext->Intermediate.NumInstructions; ++i)
+		{
+			FVectorVMOptimizeInstruction *InsStart = OptContext->Intermediate.Instructions + i;
+			if (InsStart->OpCat == EVectorVMOpCategory::Output && InsStart->Output.CopyFromInputInsIdx == -1)
+			{
+				for (uint32 j = i + 1; j < OptContext->Intermediate.NumInstructions; ++j)
+				{
+					FVectorVMOptimizeInstruction *InsEnd = OptContext->Intermediate.Instructions + j;
+					if (InsEnd->OpCat != EVectorVMOpCategory::Output || InsEnd->Output.CopyFromInputInsIdx != -1)
+					{
+						if (j - i > 1)
+						{
+							//these instructions are more than one apart so we can group them.
+							int StartInstructionIdx = i;
+							int LastInstructionIdx  = j - 1;
+							int NumInstructions = j - i;
+							FVectorVMOptimizeInstruction *StartInstruction = OptContext->Intermediate.Instructions + StartInstructionIdx;
+							{ // small list, very stupid bubble sort
+								bool sorted = false;
+								while (!sorted) {
+									sorted = true;
+									for (int k = 0; k < NumInstructions - 1; ++k) {
+										FVectorVMOptimizeInstruction *Ins0 = StartInstruction + k;
+										FVectorVMOptimizeInstruction *Ins1 = Ins0 + 1;
+										uint64 Key0 = VVMOutputInsGetSortKey(OptContext->Intermediate.SSARegisterUsageBuffer, Ins0);
+										uint64 Key1 = VVMOutputInsGetSortKey(OptContext->Intermediate.SSARegisterUsageBuffer, Ins1);
+
+										if (Key1 < Key0) {
+											FVectorVMOptimizeInstruction Temp = *Ins0;
+											*Ins0 = *Ins1;
+											*Ins1 = Temp;
+											sorted = false;
+										}
+									}
+								}
+							}
+							i = j;
+						}
+						break;
+					}
+				}
+			}
+		}
+	}
+
+
+	{ //Step 16: correct the register fuse buffer to make sure the instruction indices match re-ordered inputs
+		int32 *TempInputRegisterFuseBuffer = (int32 *)OptContext->Init.ReallocFn(nullptr, sizeof(int32) * OptContext->Intermediate.NumRegistersUsed, __FILE__, __LINE__);
+		if (TempInputRegisterFuseBuffer == nullptr)
+		{
+			return VectorVMOptimizerSetError(OptContext, VVMOptErr_OutOfMemory | VVMOptErr_InputFuseBuffer | VVMOptErr_Fatal);
+		}
+		VVMOptRAIIPtrToFree RegStackRAII(OptContext, TempInputRegisterFuseBuffer);
+		FMemory::Memcpy(TempInputRegisterFuseBuffer, OptContext->Intermediate.InputRegisterFuseBuffer, sizeof(int32) * OptContext->Intermediate.NumRegistersUsed);
+
+		for (uint32 InputInsIdx = 0; InputInsIdx < OptContext->Intermediate.NumInstructions; ++InputInsIdx)
+		{
+			FVectorVMOptimizeInstruction *InputIns = OptContext->Intermediate.Instructions + InputInsIdx;
+			if (InputIns->OpCat == EVectorVMOpCategory::Input)
+			{
+				if (InputIns->Index != InputInsIdx)
+				{ //only worry about instructions that are not in their original place
+					//fixup register fuse buffer
+					for (uint32 i = 0; i < OptContext->Intermediate.NumRegistersUsed; ++i)
+					{
+						if (TempInputRegisterFuseBuffer[i] == InputIns->Index)
+						{
+							OptContext->Intermediate.InputRegisterFuseBuffer[i] = InputInsIdx;
+						}
+					}
+				}
+			}
+		}	
+	}
+
+	{ //Step 17: use the SSA registers to compute the minimized registers required and write them back into the register usage buffer
+		int MaxLiveRegisters = 0;
+		uint16 *SSAUseMap = (uint16 *)OptContext->Init.ReallocFn(nullptr, sizeof(uint16) * NumSSARegistersUsed, __FILE__, __LINE__);
+		if (SSAUseMap == nullptr)
+		{
+			return VectorVMOptimizerSetError(OptContext, VVMOptErr_OutOfMemory | VVMOptErr_SSARemap | VVMOptErr_Fatal);
+		}
+		VVMOptRAIIPtrToFree RegStackRAII(OptContext, SSAUseMap);
+		FMemory::Memset(SSAUseMap, 0xFF, sizeof(uint16) * NumSSARegistersUsed);
+
+		FVectorVMOptimizeInsRegUsage InsRegUse;
+		FVectorVMOptimizeInsRegUsage InsRegUse2;
+		for (uint32 i = 0; i < OptContext->Intermediate.NumInstructions; ++i)
+		{
+			FVectorVMOptimizeInstruction *Ins = OptContext->Intermediate.Instructions + i;
+			GetRegistersUsedForInstruction(OptContext, Ins, &InsRegUse);
+
+			//check to see if any of the inputs are ever used again
+			for (int j = 0; j < InsRegUse.NumInputRegisters; ++j)
+			{
+				bool SSARegStillLive = false;
+				uint16 SSAInputReg = OptContext->Intermediate.SSARegisterUsageBuffer[InsRegUse.RegIndices[j]];
+				//we need to check this instruction too, because if its output aliases with its input we can't mark it as unused
+				//first check if the input and output alias, if they do, the SSA register is still active
+				//next check the instructions after this one
+				for (uint32 i2 = i + 1; i2 < OptContext->Intermediate.NumInstructions; ++i2)
+				{
+					FVectorVMOptimizeInstruction *Ins2 = OptContext->Intermediate.Instructions + i2;
+					int NumRegisters = GetRegistersUsedForInstruction(OptContext, Ins2, &InsRegUse2);
+					for (int k = 0; k < NumRegisters; ++k)
+					{
+						if (OptContext->Intermediate.SSARegisterUsageBuffer[InsRegUse2.RegIndices[k]] == SSAInputReg)
+						{
+							SSARegStillLive = true;
+							break;
+						}
+					}
+				}
+				if (!SSARegStillLive)
+				{
+					//register is no longer required, so mark it as free to use
+					for (int k = 0; k < NumSSARegistersUsed; ++k)
+					{
+						if (SSAUseMap[k] == SSAInputReg)
+						{
+							SSAUseMap[k] = 0xFFFF;
+							break;
+						}
+					}
+				}
+			}
+
+			for (int j = 0; j < InsRegUse.NumOutputRegisters; ++j)
+			{
+				uint16 SSARegIdx = OptContext->Intermediate.SSARegisterUsageBuffer[InsRegUse.RegIndices[InsRegUse.NumInputRegisters + j]];
+				uint16 MinimizedRegIdx = 0xFFFF;
+				for (uint16 k = 0; k < NumSSARegistersUsed; ++k)
+				{
+					if (SSAUseMap[k] == 0xFFFF)
+					{
+						SSAUseMap[k] = SSARegIdx;
+						MinimizedRegIdx = k;
+						break;
+					}
+				}
+				check(MinimizedRegIdx != 0xFFFF);
+				uint16 OutputRegIdx = InsRegUse.RegIndices[InsRegUse.NumInputRegisters + j];
+				OptContext->Intermediate.RegisterUsageBuffer[OutputRegIdx] = MinimizedRegIdx;
+
+				//change all future instructions to use minimized register index
+				for (uint32 i2 = i + 1; i2 < OptContext->Intermediate.NumInstructions; ++i2)
+				{
+					FVectorVMOptimizeInstruction *Ins2 = OptContext->Intermediate.Instructions + i2;
+					GetRegistersUsedForInstruction(OptContext, Ins2, &InsRegUse2);
+					for (int k = 0; k < InsRegUse2.NumInputRegisters; ++k)
+					{
+						if (OptContext->Intermediate.SSARegisterUsageBuffer[InsRegUse2.RegIndices[k]] == OptContext->Intermediate.SSARegisterUsageBuffer[OutputRegIdx])
+						{
+							OptContext->Intermediate.RegisterUsageBuffer[InsRegUse2.RegIndices[k]] = MinimizedRegIdx;
+						}
+					}
+				}
+			}
+
+			{ //count the live registers
+				int NumLiveRegisters = 0;
+				for (int j = 0; j < NumSSARegistersUsed; ++j)
+				{
+					NumLiveRegisters += (SSAUseMap[j] != 0xFFFF);
+				}
+				if (NumLiveRegisters > MaxLiveRegisters)
+				{
+					MaxLiveRegisters = NumLiveRegisters;
+				}
+			}
+		}
+		OptContext->NumTempRegisters = (uint32)MaxLiveRegisters;
+	}
+
+	{ //Step 18: write the final optimized bytecode
+		//this goes over the instruction list twice.  The first time to figure out how many bytes are required for the bytecode, the second to write the bytecode.
+		uint8 *OptimizedBytecode = nullptr;
+		int NumOptimizedBytesRequired = 0;
+		int NumOptimizedBytesWritten = 0;
+#		define VVMOptWriteByte(b)   if (OptimizedBytecode)                                                  \
+									{                                                                       \
+										check(NumOptimizedBytesWritten <= NumOptimizedBytesRequired - 1);   \
+                                        OptimizedBytecode[NumOptimizedBytesWritten++] = (uint8)b;           \
+                                    }                                                                       \
+									else                                                                    \
+									{                                                                       \
+                                        ++NumOptimizedBytesRequired;                                        \
+                                    }
+#		define VVMOptWriteU16(b)    if (OptimizedBytecode)                                                  \
+									{                                                                       \
+										check(NumOptimizedBytesWritten <= NumOptimizedBytesRequired - 2);   \
+                                        OptimizedBytecode[NumOptimizedBytesWritten++] = (uint8)(b & 0xFF);  \
+										OptimizedBytecode[NumOptimizedBytesWritten++] = (uint8)(b >> 8);    \
+									}                                                                       \
+									else                                                                    \
+									{                                                                       \
+										NumOptimizedBytesRequired += 2;                                     \
+									}
+		WriteOptimizedBytecode:
+		for (uint32 i = 0; i < OptContext->Intermediate.NumInstructions; ++i)
+		{
+			FVectorVMOptimizeInstruction *Ins = OptContext->Intermediate.Instructions + i;
+			if (OptimizedBytecode)
+			{
+				Ins->PtrOffsetInOptimizedBytecode = (uint32)NumOptimizedBytesWritten;
+			}
+			switch (Ins->OpCat)
+			{
+				case EVectorVMOpCategory::Input:
+					if (Ins->Input.FirstInsInsertIdx != -1)
+					{
+						VVMOptWriteByte(Ins->OpCode);
+						VVMOptWriteU16(Ins->Input.DataSetIdx);
+						VVMOptWriteU16(Ins->Input.InputIdx);
+						VVMOptWriteU16(OptContext->Intermediate.RegisterUsageBuffer[Ins->Input.DstRegPtrOffset]);
+					}
+					else
+					{
+						Ins->PtrOffsetInOptimizedBytecode = -1;
+					}
+					break;
+				case EVectorVMOpCategory::Output:
+					if (Ins->Output.CopyFromInputInsIdx == -1)
+					{
+						int NumOutputInstructions = 1;
+						//figure out how we can batch these
+						for (uint32 j = i + 1; j < OptContext->Intermediate.NumInstructions; ++j)
+						{
+							FVectorVMOptimizeInstruction *NextIns = OptContext->Intermediate.Instructions + j;
+							if (NextIns->OpCode == Ins->OpCode                                                                                                                              &&
+								NextIns->Output.CopyFromInputInsIdx == -1                                                                                                                   &&
+								NextIns->Output.DataSetIdx == Ins->Output.DataSetIdx                                                                                                        &&
+								OptContext->Intermediate.SSARegisterUsageBuffer[NextIns->Output.RegPtrOffset] == OptContext->Intermediate.SSARegisterUsageBuffer[Ins->Output.RegPtrOffset])
+							{
+								++NumOutputInstructions;
+								if (NumOutputInstructions >= 0xFF)
+								{
+									break; //we only write 1 byte so we can't group anymore
+								}
+							}
+							else
+							{
+								break;
+							}
+						}
+						const int TotalNumOutputInstructions = NumOutputInstructions;
+						//batched instructions have slightly weird bytcode because all the extra data required for the instruction come after the registers, (DataSet index, float/int flag)
+						//This is so decoding can be done 4-at-a-time using the standard decoding method in the VM Exec().  The index register is different for all output_batch instructions
+						//so it can be most efficiently decoded
+						//output_batch4/8 puts the index at the very end and must be decoded separately.  It's guaranteed not to be constant so it can be decoded very efficiently
+						//output_batch3/7 and output_batch2 put the index first so it's automatically decoded correctly
+						while (NumOutputInstructions > 0)
+						{
+							int StartNumOutputInstructions = NumOutputInstructions;
+							if (NumOutputInstructions >= 8 && (OptContext->Intermediate.RegisterUsageBuffer[Ins->Output.RegPtrOffset] & 0x8000) == 0)
+							{
+								VVMOptWriteByte((uint8)EVectorVMOp::output_batch8);
+								for (int j = 0; j < 8; ++j)
+								{
+									VVMOptWriteU16(OptContext->Intermediate.RegisterUsageBuffer[Ins[j].Output.RegPtrOffset + 1]);  //bytes 0-15: input indices
+								}
+								for (int j = 0; j < 8; ++j)
+								{
+									VVMOptWriteU16(Ins[j].Output.DstRegIdx);                                                       //bytes 16-31: output indices
+								}
+								VVMOptWriteU16(Ins->Output.DataSetIdx);                                                            //bytes 32-33. DataSet index
+								VVMOptWriteU16(OptContext->Intermediate.RegisterUsageBuffer[Ins->Output.RegPtrOffset]);            //bytes 34-35: index reg
+								VVMOptWriteByte((uint8)Ins->OpCode - (uint8)EVectorVMOp::outputdata_float);                        //byte  36: float/int
+								NumOutputInstructions -= 8;
+							}
+							else if (NumOutputInstructions >= 7)
+							{
+								VVMOptWriteByte((uint8)EVectorVMOp::output_batch7);
+								VVMOptWriteU16(OptContext->Intermediate.RegisterUsageBuffer[Ins->Output.RegPtrOffset]);            //bytes 0-1: index reg
+								for (int j = 0; j < 3; ++j)
+								{
+									VVMOptWriteU16(OptContext->Intermediate.RegisterUsageBuffer[Ins[j].Output.RegPtrOffset + 1]);  //bytes 2-7: input indices
+								}
+								VVMOptWriteU16(Ins->Output.DataSetIdx);                                                            //bytes 8-9. DataSet index
+								for (int j = 0; j < 4; ++j)
+								{
+									VVMOptWriteU16(OptContext->Intermediate.RegisterUsageBuffer[Ins[j + 3].Output.RegPtrOffset + 1]);  //bytes 10-17: input indices
+								}
+								for (int j = 0; j < 7; ++j)
+								{
+									VVMOptWriteU16(Ins[j].Output.DstRegIdx);                                                       //bytes 18-31: output indices
+								}
+								VVMOptWriteByte((uint8)Ins->OpCode - (uint8)EVectorVMOp::outputdata_float);                        //byte  32: float/int
+								NumOutputInstructions -= 7;
+							}
+							else if (NumOutputInstructions >= 4 && (OptContext->Intermediate.RegisterUsageBuffer[Ins->Output.RegPtrOffset] & 0x8000) == 0)
+							{
+								VVMOptWriteByte((uint8)EVectorVMOp::output_batch4);
+								for (int j = 0; j < 4; ++j)
+								{
+									VVMOptWriteU16(OptContext->Intermediate.RegisterUsageBuffer[Ins[j].Output.RegPtrOffset + 1]);  //bytes 0-7: input indices
+								}
+								for (int j = 0; j < 4; ++j)
+								{
+									VVMOptWriteU16(Ins[j].Output.DstRegIdx);                                                       //bytes 8-15: output indices
+								}
+								VVMOptWriteU16(Ins->Output.DataSetIdx);                                                            //bytes 16-17. DataSet index
+								VVMOptWriteU16(OptContext->Intermediate.RegisterUsageBuffer[Ins->Output.RegPtrOffset]);            //bytes 18-19: index reg
+								VVMOptWriteByte((uint8)Ins->OpCode - (uint8)EVectorVMOp::outputdata_float);                        //byte  20: float/int
+								NumOutputInstructions -= 4;
+							}
+							else if (NumOutputInstructions >= 3)
+							{
+								VVMOptWriteByte((uint8)EVectorVMOp::output_batch3);
+								VVMOptWriteU16(OptContext->Intermediate.RegisterUsageBuffer[Ins->Output.RegPtrOffset]);            //bytes 0-1: index reg
+								for (int j = 0; j < 3; ++j)
+								{
+									VVMOptWriteU16(OptContext->Intermediate.RegisterUsageBuffer[Ins[j].Output.RegPtrOffset + 1]);  //bytes 2-7: input indices
+								}
+								VVMOptWriteU16(Ins->Output.DataSetIdx);                                                            //bytes 8-9. DataSet index
+								for (int j = 0; j < 3; ++j)
+								{
+									VVMOptWriteU16(Ins[j].Output.DstRegIdx);                                                       //bytes 10-15: output indices
+								}
+								VVMOptWriteByte((uint8)Ins->OpCode - (uint8)EVectorVMOp::outputdata_float);                        //byte  16: float/int
+								NumOutputInstructions -= 3;
+							}
+							else if (NumOutputInstructions >= 2)
+							{
+								VVMOptWriteByte((uint8)EVectorVMOp::output_batch2);
+								VVMOptWriteU16(OptContext->Intermediate.RegisterUsageBuffer[Ins->Output.RegPtrOffset]);            //bytes 0-1: index reg
+								for (int j = 0; j < 2; ++j)
+								{
+									VVMOptWriteU16(OptContext->Intermediate.RegisterUsageBuffer[Ins[j].Output.RegPtrOffset + 1]);  //bytes 2-5: input indices
+								}
+								VVMOptWriteU16(Ins->Output.DataSetIdx);                                                            //bytes 6-9. DataSet index
+								for (int j = 0; j < 2; ++j)
+								{
+									VVMOptWriteU16(Ins[j].Output.DstRegIdx);                                                       //bytes 10-11: output indices
+								}
+								VVMOptWriteByte((uint8)Ins->OpCode - (uint8)EVectorVMOp::outputdata_float);                        //byte  12: float/int
+								NumOutputInstructions -= 2;
+							}
+							else
+							{
+								check(NumOutputInstructions == 1);
+								VVMOptWriteByte(Ins->OpCode);
+								VVMOptWriteU16(Ins->Output.DataSetIdx);                                                      //0: DataSet index
+								VVMOptWriteU16(OptContext->Intermediate.RegisterUsageBuffer[Ins->Output.RegPtrOffset]);      //1: index reg
+								VVMOptWriteU16(OptContext->Intermediate.RegisterUsageBuffer[Ins->Output.RegPtrOffset + 1]);  //2: input reg
+								VVMOptWriteU16(Ins->Output.DstRegIdx);
+								--NumOutputInstructions;
+							}
+							int NumOutputsConsumed = StartNumOutputInstructions - NumOutputInstructions;
+							Ins += NumOutputsConsumed;
+						}
+						i += TotalNumOutputInstructions - 1;
+						check(NumOutputInstructions == 0);
+					}
+					else
+					{ //copy_to_output, bypass the temp registers entirely
+						FVectorVMOptimizeInstruction *InputIns = OptContext->Intermediate.Instructions + Ins->Output.CopyFromInputInsIdx;
+						uint8 InputRegType  = (uint8)InputIns->OpCode - (uint8)EVectorVMOp::inputdata_float;
+						uint8 OutputRegType = (uint8)Ins->OpCode      - (uint8)EVectorVMOp::outputdata_float;
+						check(InputRegType == OutputRegType);
+						VVMOptWriteByte(EVectorVMOp::copy_to_output);
+						VVMOptWriteU16(Ins->Output.DataSetIdx);
+						VVMOptWriteU16(InputIns->Input.DataSetIdx);
+						VVMOptWriteU16(OptContext->Intermediate.RegisterUsageBuffer[Ins->Output.RegPtrOffset]);
+						VVMOptWriteByte(InputRegType);
+						uint8 *CountPtr;
+						uint8 TempCount;
+						if (OptimizedBytecode)
+						{
+							CountPtr = OptimizedBytecode + NumOptimizedBytesWritten++;
+						}
+						else
+						{
+							CountPtr = &TempCount;
+							++NumOptimizedBytesRequired;
+						}
+						*CountPtr = 1;
+					
+						VVMOptWriteU16(Ins->Output.DstRegIdx);
+						VVMOptWriteU16(InputIns->Input.InputIdx);
+					
+						//merge all the subsequent copy_to_output instructions that share the same output data set, input data set and register type (output ops should be sorted above)
+						while (i < OptContext->Intermediate.NumInstructions - 1 && *CountPtr < 0xFF)
+						{
+							FVectorVMOptimizeInstruction *NextIns = OptContext->Intermediate.Instructions + i + *CountPtr;
+							if (NextIns->OpCat == EVectorVMOpCategory::Output && NextIns->Output.CopyFromInputInsIdx != -1)
+							{
+								FVectorVMOptimizeInstruction *NextInputIns = OptContext->Intermediate.Instructions + NextIns->Output.CopyFromInputInsIdx;
+								uint8 NextInputRegType                     = (uint8)NextInputIns->OpCode - (uint8)EVectorVMOp::inputdata_float;
+								uint8 NextOutputRegType                    = (uint8)NextIns->OpCode - (uint8)EVectorVMOp::outputdata_float;
+								check(NextInputRegType == NextOutputRegType);
+								check(NextOutputRegType == 0 || NextOutputRegType == 1);
+								if (NextIns->Output.DataSetIdx     == Ins->Output.DataSetIdx     &&
+									NextInputIns->Input.DataSetIdx == InputIns->Input.DataSetIdx &&
+									NextOutputRegType              == OutputRegType)
+								{
+									VVMOptWriteU16(NextIns->Output.DstRegIdx);
+									VVMOptWriteU16(NextInputIns->Input.InputIdx);
+									++*CountPtr;
+								}
+								else
+								{
+									break;
+								}
+							}
+							else
+							{
+								break;
+							}
+						}
+						i += *CountPtr - 1; //skip the copy-to-output instructions we merged into this one
+					}
+				break;
+				case EVectorVMOpCategory::Op:
+					if (Ins->Op.InputFuseBits == 0)
+					{ //all inputs are regular registers, write the operation as normal
+						VVMOptWriteByte(Ins->OpCode);
+						for (int j = 0; j < Ins->Op.NumInputs + Ins->Op.NumOutputs; ++j) {
+							VVMOptWriteU16(OptContext->Intermediate.RegisterUsageBuffer[Ins->Op.RegPtrOffset + j]);
+						}
+					}
+					else
+					{ //at least one of the inputs should be from a dataset, not a register
+						check(Ins->Op.NumInputs > 0);
+						check(Ins->Op.NumInputs <= 3);
+						EVectorVMOp StartOp[] = { EVectorVMOp::fused_input1_1, EVectorVMOp::fused_input2_1, EVectorVMOp::fused_input3_1 };
+						EVectorVMOp FusedOp = (EVectorVMOp)((uint8)StartOp[Ins->Op.NumInputs - 1] - 1 + Ins->Op.InputFuseBits);
+						if (FusedOp == EVectorVMOp::fused_input3_4) //the bit pattern doesn't match for these two ops only due to making the decoder more efficient so do that here
+						{ 
+							FusedOp = EVectorVMOp::fused_input3_3;
+						}
+						else if (FusedOp == EVectorVMOp::fused_input3_3)
+						{
+							FusedOp = EVectorVMOp::fused_input3_4;
+						}
+						VVMOptWriteByte(FusedOp);
+						{ //write all the inputs as normal so they get decoded correctly in the universal decoder
+							for (int j = 0; j < Ins->Op.NumInputs + Ins->Op.NumOutputs; ++j)
+							{
+								if (Ins->Op.InputFuseBits & (1 << j))
+								{
+									int InputInsIdx = OptContext->Intermediate.InputRegisterFuseBuffer[Ins->Op.RegPtrOffset + j];
+									check(InputInsIdx != -1);
+									FVectorVMOptimizeInstruction *InputIns = OptContext->Intermediate.Instructions + InputInsIdx;
+									VVMOptWriteU16(InputIns->Input.InputIdx);
+								}
+								else
+								{
+									VVMOptWriteU16(OptContext->Intermediate.RegisterUsageBuffer[Ins->Op.RegPtrOffset + j]);
+								}
+							}
+						}
+						{ //write the real op next, independent of how many more bytes we need to read
+							VVMOptWriteByte(Ins->OpCode);
+							for (uint16 j = 0; j < Ins->Op.NumInputs; ++j)
+							{
+								if (Ins->Op.InputFuseBits & (1 << j))
+								{
+									int InputInsIdx = OptContext->Intermediate.InputRegisterFuseBuffer[Ins->Op.RegPtrOffset + j];
+									check(InputInsIdx != -1);
+									FVectorVMOptimizeInstruction *InputIns = OptContext->Intermediate.Instructions + InputInsIdx;
+									check(InputIns->OpCode == EVectorVMOp::inputdata_float || InputIns->OpCode == EVectorVMOp::inputdata_int32);
+									VVMOptWriteByte((uint8)InputIns->OpCode - (uint8)EVectorVMOp::inputdata_float);
+									VVMOptWriteU16(InputIns->Input.DataSetIdx);
+								}
+							}
+						}
+					}
+					break;
+				case EVectorVMOpCategory::IndexGen:
+					VVMOptWriteByte(Ins->OpCode);
+					VVMOptWriteU16(Ins->IndexGen.DataSetIdx);                                                     //0: DataSetIdx
+					VVMOptWriteU16(OptContext->Intermediate.RegisterUsageBuffer[Ins->IndexGen.RegPtrOffset + 0]); //1: Input Register
+					VVMOptWriteU16(OptContext->Intermediate.RegisterUsageBuffer[Ins->IndexGen.RegPtrOffset + 1]); //2: Write-gather Output Register
+					VVMOptWriteU16(OptContext->Intermediate.RegisterUsageBuffer[Ins->IndexGen.RegPtrOffset + 2]); //3: Original VM Output Register (if 0xFFFF then nothing)
+					break;
+				case EVectorVMOpCategory::ExtFnCall:
+					VVMOptWriteByte(Ins->OpCode);
+					VVMOptWriteU16(Ins->ExtFnCall.ExtFnIdx);
+					for (int j = 0; j < ExtFnIOData[Ins->ExtFnCall.ExtFnIdx].NumInputs + ExtFnIOData[Ins->ExtFnCall.ExtFnIdx].NumOutputs; ++j)
+					{
+						VVMOptWriteU16(OptContext->Intermediate.RegisterUsageBuffer[Ins->ExtFnCall.RegPtrOffset + j]);
+					}
+					break;
+				case EVectorVMOpCategory::ExecIndex: 
+					VVMOptWriteByte(Ins->OpCode);
+					VVMOptWriteU16(OptContext->Intermediate.RegisterUsageBuffer[Ins->ExecIndex.RegPtrOffset]);
+					break;
+				case EVectorVMOpCategory::Stat:
+					if (!(Flags & VVMOptFlag_OmitStats))
+					{
+						VVMOptWriteByte(Ins->OpCode);
+						if (Ins->OpCode == EVectorVMOp::enter_stat_scope)
+						{
+							VVMOptWriteU16(Ins->Stat.ID);
+						}
+						else if (Ins->OpCode == EVectorVMOp::exit_stat_scope)
+						{
+							
+						} 
+						else
+						{
+							check(false);
+						}
+					} else {
+						Ins->PtrOffsetInOptimizedBytecode = -1;
+					}
+					break;
+				case EVectorVMOpCategory::RWBuffer:
+					check(Ins->OpCode == EVectorVMOp::update_id || Ins->OpCode == EVectorVMOp::acquire_id);
+					VVMOptWriteByte(Ins->OpCode);
+					VVMOptWriteU16(Ins->RWBuffer.DataSetIdx);
+					VVMOptWriteU16(OptContext->Intermediate.RegisterUsageBuffer[Ins->RWBuffer.RegPtrOffset + 0]);
+					VVMOptWriteU16(OptContext->Intermediate.RegisterUsageBuffer[Ins->RWBuffer.RegPtrOffset + 1]);
+					break;
+				case EVectorVMOpCategory::Other:
+					switch (Ins->OpCode)
+					{
+						case EVectorVMOp::done:
+							check(i == OptContext->Intermediate.NumInstructions - 1);
+							break;
+						case EVectorVMOp::noise2D:
+							VVMOptWriteByte(Ins->OpCode);
+							check(false);
+							break;
+						case EVectorVMOp::noise3D:
+							VVMOptWriteByte(Ins->OpCode);
+							check(false);
+							break;
+					}
+					break;
+
+			}
+		}
+		if (OptimizedBytecode == nullptr)
+		{
+			check(NumOptimizedBytesWritten == 0);
+			if (NumOptimizedBytesRequired > 0)
+			{
+				OptimizedBytecode = (uint8 *)OptContext->Init.ReallocFn(nullptr, NumOptimizedBytesRequired + 16, __FILE__, __LINE__); //we decode 4 registers at a time so make sure to pad the allocation so no matter what we don't read past the end
+				if (OptimizedBytecode == nullptr)
+				{
+					return VectorVMOptimizerSetError(OptContext, VVMOptErr_OutOfMemory | VVMOptErr_OptimizedBytecode | VVMOptErr_Fatal);
+				}
+				else
+				{
+					goto WriteOptimizedBytecode;
+				}
+			}
+		}
+		check(NumOptimizedBytesWritten == NumOptimizedBytesRequired);
+#		undef VVMOptWriteByte
+#		undef VVMOptWriteU16
+
+		OptContext->OutputBytecode   = OptimizedBytecode;
+		OptContext->NumBytecodeBytes = NumOptimizedBytesWritten;
+	}
+
+	if (!(Flags & VVMOptFlag_SaveIntermediateState))
+	{
+		VectorVMFreeOptimizerIntermediateData(OptContext);
+	}
+
+#undef AllocBytecodeBytes
+#undef AllocRegisterUse
+#undef VVMOptimizeWriteRegIndex
+#undef VVMOptimizeWrite2
+#undef VVMOptimizeVecIns1
+#undef VVMOptimizeVecIns2
+#undef VVMOptimizeVecIns3
+	return 0;
+}
+
+
+#undef VectorVMOptimizerSetError
