- New Vector VM with optimizer and serialization.

diff --git a/Engine/Source/Runtime/VectorVM/Private/VectorVMExperimental_Serialization.inl b/Engine/Source/Runtime/VectorVM/Private/VectorVMExperimental_Serialization.inl
new file mode 100644
--- /dev/null
+++ b/Engine/Source/Runtime/VectorVM/Private/VectorVMExperimental_Serialization.inl
@@ -0,0 +1,874 @@
+// Copyright Epic Games, Inc. All Rights Reserved.
+
+#ifdef VVM_INCLUDE_SERIALIZATION
+
+//prototypes for stuff shared between the original and experimental VM
+static FVectorVMSerializeInstruction *VVMSerGetNextInstruction(FVectorVMSerializeState *SerializeState, int InstructionIdx, int GlobalChunkIdx, int BatchUseCount, int OpStart, int NumOps, uint64 Dt, uint64 DtDecode);
+
+VECTORVM_API void FreeVectorVMSerializeState(FVectorVMSerializeState *SerializeState)
+{
+	if (SerializeState->FreeFn == nullptr)
+	{
+		return;
+	}
+	for (uint32 i = 0; i < SerializeState->NumExternalData; ++i)
+	{
+		if (SerializeState->ExternalData[i].Name)
+		{
+			SerializeState->FreeFn(SerializeState->ExternalData[i].Name, __FILE__, __LINE__);
+		}
+	}
+	SerializeState->FreeFn(SerializeState->TempRegFlags, __FILE__, __LINE__);
+	SerializeState->FreeFn(SerializeState->ExternalData, __FILE__, __LINE__);
+	SerializeState->FreeFn(SerializeState->Bytecode    , __FILE__, __LINE__);
+	SerializeState->FreeFn(SerializeState->Chunks      , __FILE__, __LINE__);
+
+	for (uint32 i = 0; i < SerializeState->NumDataSets; ++i)
+	{
+		if (SerializeState->DataSets[i].InputBuffers)
+		{
+			SerializeState->FreeFn(SerializeState->DataSets[i].InputBuffers, __FILE__, __LINE__);
+		}
+		if (SerializeState->DataSets[i].OutputBuffers)
+		{
+			SerializeState->FreeFn(SerializeState->DataSets[i].OutputBuffers, __FILE__, __LINE__);
+		}		
+		if (SerializeState->DataSets[i].InputIDTable)
+		{
+			SerializeState->FreeFn(SerializeState->DataSets[i].InputIDTable, __FILE__, __LINE__);
+		}
+		if (SerializeState->DataSets[i].InputFreeIDTable)
+		{
+			SerializeState->FreeFn(SerializeState->DataSets[i].InputFreeIDTable, __FILE__, __LINE__);
+		}
+		if (SerializeState->DataSets[i].InputSpawnedIDTable)
+		{
+			SerializeState->FreeFn(SerializeState->DataSets[i].InputSpawnedIDTable, __FILE__, __LINE__);
+		}
+		if (SerializeState->DataSets[i].OutputIDTable)
+		{
+			SerializeState->FreeFn(SerializeState->DataSets[i].OutputIDTable, __FILE__, __LINE__);
+		}
+		if (SerializeState->DataSets[i].OutputFreeIDTable)
+		{
+			SerializeState->FreeFn(SerializeState->DataSets[i].OutputFreeIDTable, __FILE__, __LINE__);
+		}
+		if (SerializeState->DataSets[i].OutputSpawnedIDTable)
+		{
+			SerializeState->FreeFn(SerializeState->DataSets[i].OutputSpawnedIDTable, __FILE__, __LINE__);
+		}
+	}
+	SerializeState->FreeFn(SerializeState->DataSets, __FILE__, __LINE__);
+	for (uint32 i = 0; i < SerializeState->NumInstructions; ++i)
+	{
+		SerializeState->FreeFn(SerializeState->Instructions[i].TempRegisters, __FILE__, __LINE__);
+		SerializeState->FreeFn(SerializeState->Instructions[i].TempRegisterFlags, __FILE__, __LINE__);
+	}
+	SerializeState->FreeFn(SerializeState->Instructions, __FILE__, __LINE__);
+	SerializeState->FreeFn(SerializeState->PreExecConstData, __FILE__, __LINE__);
+	SerializeState->FreeFn(SerializeState->PostExecConstData, __FILE__, __LINE__);
+}
+
+static uint32 VectorVMSerializeSetError_(FVectorVMSerializeState *SerializeState, uint32 Flags, uint32 LineNum)
+{
+	SerializeState->Error.Line = LineNum;
+	if (SerializeState->Error.CallbackFn)
+	{
+		SerializeState->Error.Flags = SerializeState->Error.CallbackFn(SerializeState, SerializeState->Error.Flags | Flags);
+	}
+	else
+	{
+		SerializeState->Error.Flags |= Flags;
+	}
+	if (SerializeState->Error.Flags & VVMSerErr_Fatal)
+	{
+		check(false); //hit the debugger
+		FreeVectorVMSerializeState(SerializeState);
+	}
+	return SerializeState->Error.Flags;
+}
+
+#define VectorVMSerializeSetError(SerializeState, Flags)     VectorVMSerializeSetError_(SerializeState, Flags, __LINE__)
+
+
+#ifdef NIAGARA_EXP_VM
+
+/*********************************************************************************************************************************************************************************************************************************
+*** Serialization stuff for the Experimental VM
+*********************************************************************************************************************************************************************************************************************************/
+#define VVMSer_batchStart(...)
+#define VVMSer_batchEnd(...)
+
+#define VVMSer_chunkStart(...)
+#define VVMSer_chunkEnd(...)
+
+#define VVMSer_insStart(...)
+#define VVMSer_insEndDecode(...)
+#define VVMSer_insEnd(...)
+
+#define VVMSer_batchStartExp(SerializeState, ...)
+#define VVMSer_batchEndExp(SerializeState, ...)
+
+#ifdef VVM_SERIALIZE_NO_WRITE
+#define VVMSer_chunkStartExp(...)
+#define VVMSer_chunkEndExp(...)
+#define VVMSer_insStartExp(...)
+#define VVMSer_insEndDecodeExp(...)
+#define VVMSer_insEndExp(...)
+#define VVMSer_initSerializationState(...)
+#else //VVM_SERIALIZE_NO_WRITE
+
+
+#define VVMSer_chunkStartExp(SerializeState, ChunkIdx_, BatchIdx)                   \
+	FVectorVMSerializeChunk *SerializeChunk = nullptr;                              \
+	int VVMSerNumInstructionsThisChunk = 0;                                         \
+	int VVMSerGlobalChunkIdx = ExecIdx * VVMState->MaxChunksPerBatch + ChunkIdx_;   \
+	if (SerializeState)                                                             \
+	{                                                                               \
+		SerializeChunk = SerializeState->Chunks + VVMSerGlobalChunkIdx;             \
+		SerializeChunk->ChunkIdx = VVMSerGlobalChunkIdx;                            \
+		SerializeChunk->BatchIdx = BatchIdx;                                        \
+		SerializeChunk->NumInstances = NumInstancesThisChunk;                       \
+		SerializeChunk->StartInstance = StartInstanceThisChunk;                     \
+		SerializeChunk->StartThreadID = FPlatformTLS::GetCurrentThreadId();         \
+		SerializeChunk->StartClock = FPlatformTime::Cycles64();                     \
+	}
+
+#define VVMSer_chunkEndExp(SerializeState, ...)                                                                                 \
+	if (SerializeChunk)                                                                                                         \
+	{                                                                                                                           \
+		FPlatformAtomics::InterlockedOr(&SerializeState->ChunkComplete, (int64)((1ULL << (uint64)SerializeChunk->ChunkIdx)));   \
+		SerializeChunk->EndThreadID = FPlatformTLS::GetCurrentThreadId();                                                       \
+		SerializeChunk->EndClock = FPlatformTime::Cycles64();                                                                   \
+	}
+
+#define VVMSer_insStartExp(SerializeState, ...)             \
+	const uint8 *VVMSerStartOpPtr = InsPtr;                 \
+	uint64 VVMSerStartCycles = FPlatformTime::Cycles64();
+
+
+#define VVMSer_insEndDecodeExp(SerializeState, ...)             \
+	uint64 VVMSerEndDecodeCycles = FPlatformTime::Cycles64();
+
+#define VVMSer_insEndExp(SerializeState, OpStart, NumOps, ...)                                              \
+	if (SerializeState)                                                                                     \
+	{                                                                                                       \
+		uint64 VVMSerEndExecCycles = FPlatformTime::Cycles64();                                             \
+		VVMSer_serializeInstruction(SerializeState, VVMState, BatchState, VVMSerNumInstructionsThisChunk++, \
+			StartInstanceThisChunk, NumInstancesThisChunk, VVMSerGlobalChunkIdx, NumLoops,                  \
+			OpStart, NumOps, VVMSerStartCycles, VVMSerEndDecodeCycles, VVMSerEndExecCycles);                \
+		uint64 VVMSerEndSerializeCycles = FPlatformTime::Cycles64();                                        \
+		SerializeState->SerializeDt += VVMSerEndSerializeCycles - VVMSerEndExecCycles;                      \
+		SerializeChunk->InsExecTime += VVMSerEndExecCycles - VVMSerStartCycles;                             \
+	}
+
+static void VVMSer_serializeInstruction(FVectorVMSerializeState *SerializeState, FVectorVMState *VectorVMState, FVectorVMBatchState *BatchState, 
+										int InstructionIdx, int StartInstance, int NumInstancesThisChunk, int GlobalChunkIdx, int NumLoops, 
+										int OpStart, int NumOps, uint64 StartInsCycles, uint64 EndDecodeCycles, uint64 EndInsCycles)
+{
+	
+	FVectorVMSerializeInstruction *Ins = VVMSerGetNextInstruction(SerializeState, InstructionIdx, GlobalChunkIdx, BatchState->UseCount, OpStart, NumOps, EndInsCycles - StartInsCycles, EndDecodeCycles - StartInsCycles);
+	if (Ins != nullptr)
+	{
+		check(Ins->TempRegisters != nullptr);
+		for (uint32 i = 0; i < SerializeState->NumTempRegisters; ++i)
+		{
+			FMemory::Memcpy(Ins->TempRegisters + SerializeState->NumInstances * i + StartInstance, 
+				BatchState->RegisterData + VectorVMState->NumConstBuffers + i * NumLoops,
+				sizeof(uint32) * NumInstancesThisChunk);
+			check(Ins->TempRegisters != nullptr);
+			if (GlobalChunkIdx == 0)
+			{
+				Ins->TempRegisterFlags[i]  = SerializeState->TempRegFlags[i];
+			}
+		}
+	}
+}
+
+#define VVMSer_initSerializationState(VectorVMState, SerializeState, InitData, Flags) VVMSer_initSerializationState_(VectorVMState, SerializeState, InitData, (uint32)Flags)
+
+static uint32 VVMSer_initSerializationState_(FVectorVMState *VVMState, FVectorVMSerializeState *SerializeState, FVectorVMInitData *InitData, uint32 Flags)
+{
+	if (SerializeState)
+	{
+		SerializeState->Error.Flags  = 0;
+
+		SerializeState->Flags            = Flags;
+		SerializeState->NumInstances     = VVMState->TotalNumInstances;
+		SerializeState->NumTempRegisters = InitData->OptimizeContext->NumTempRegisters;
+		SerializeState->NumConstBuffers  = VVMState->NumConstBuffers;
+		SerializeState->ReallocFn        = InitData->ReallocFn ? InitData->ReallocFn : VVMDefaultRealloc;
+		SerializeState->FreeFn           = InitData->FreeFn ? InitData->FreeFn : VVMDefaultFree;
+		SerializeState->OptimizeCtx      = InitData->OptimizeContext;
+		if (SerializeState->OptimizeCtx) {
+			SerializeState->MaxExtFnRegisters = SerializeState->OptimizeCtx->MaxExtFnRegisters;
+			SerializeState->MaxExtFnUsed = SerializeState->OptimizeCtx->MaxExtFnUsed >= 0 ? (uint32)SerializeState->OptimizeCtx->MaxExtFnUsed : 0;
+		} else {
+			SerializeState->MaxExtFnRegisters = 0;
+			SerializeState->MaxExtFnUsed = 0;
+		}
+		
+
+		SerializeState->NumTempRegFlags = VVMState->NumTempRegisters;
+		for (int i = 0; i < InitData->DataSets.Num(); ++i)
+		{
+			if ((uint32)InitData->DataSets[i].InputRegisters.Num() > SerializeState->NumTempRegFlags)
+			{
+				SerializeState->NumTempRegFlags = InitData->DataSets[i].InputRegisters.Num();
+			}
+		}
+
+		SerializeState->TempRegFlags = (uint8 *)SerializeState->ReallocFn(nullptr, SerializeState->NumTempRegFlags, __FILE__, __LINE__);
+		if (SerializeState->TempRegFlags == nullptr)
+		{
+			return VectorVMSerializeSetError(SerializeState, VVMSerErr_OutOfMemory | VVMSerErr_Init | VVMSerErr_Fatal);
+		}
+		for (uint32 i = 0; i < SerializeState->NumTempRegFlags; ++i)
+		{
+			SerializeState->TempRegFlags[i] = 0;
+		}
+
+		SerializeState->NumChunks = VVMState->NumBatches * VVMState->MaxChunksPerBatch;
+		SerializeState->Chunks = (FVectorVMSerializeChunk *)SerializeState->ReallocFn(nullptr, sizeof(FVectorVMSerializeChunk) * SerializeState->NumChunks, __FILE__, __LINE__);
+		if (SerializeState->Chunks == nullptr)
+		{
+			return VectorVMSerializeSetError(SerializeState, VVMSerErr_OutOfMemory | VVMSerErr_Init | VVMSerErr_Fatal);
+		}
+		FMemory::Memset(SerializeState->Chunks, 0, sizeof(*SerializeState->Chunks) * SerializeState->NumChunks);
+		if (SerializeVectorVMInputDataSets(SerializeState, InitData->DataSets, InitData->ConstData, InitData->NumConstData) != 0)
+		{
+			return SerializeState->Error.Flags;
+		}
+		
+		
+
+		int NumExternalFunctions = VVMState->NumExtFunctions;
+		if (NumExternalFunctions != 0)
+		{
+			SerializeState->ExternalData = (FVectorVMSerializeExternalData *)SerializeState->ReallocFn(nullptr, sizeof(FVectorVMSerializeExternalData) * NumExternalFunctions, __FILE__, __LINE__);
+			if (SerializeState->ExternalData == nullptr)
+			{
+				return VectorVMSerializeSetError(SerializeState, VVMSerErr_OutOfMemory | VVMSerErr_Init | VVMSerErr_Fatal);
+			}
+			SerializeState->NumExternalData = NumExternalFunctions;
+			for (int i = 0; i < NumExternalFunctions; ++i)
+			{
+				FVectorVMExtFunctionData *f = VVMState->ExtFunctionTable + i;
+				FVectorVMSerializeExternalData *ExtData = SerializeState->ExternalData + i;
+				ExtData->Name       = nullptr;
+				ExtData->NameLen    = 0;
+				ExtData->NumInputs  = f->NumInputs;
+				ExtData->NumOutputs = f->NumOutputs;
+			}
+		}
+		else
+		{
+			SerializeState->ExternalData    = nullptr;
+			SerializeState->NumExternalData = 0;
+		}
+
+		SerializeState->Bytecode = (unsigned char *)SerializeState->ReallocFn(nullptr, InitData->OptimizeContext->NumBytecodeBytes, __FILE__, __LINE__);
+		if (SerializeState->Bytecode == nullptr)
+		{
+			return VectorVMSerializeSetError(SerializeState, VVMSerErr_OutOfMemory | VVMSerErr_Init | VVMSerErr_Fatal);
+		}
+		FMemory::Memcpy(SerializeState->Bytecode, InitData->OptimizeContext->OutputBytecode, InitData->OptimizeContext->NumBytecodeBytes);
+		SerializeState->NumBytecodeBytes = InitData->OptimizeContext->NumBytecodeBytes;
+	}
+	return 0;
+}
+#endif //VVM_SERIALIZE_NO_WRITE
+
+#else //NIAGARA_EXP_VM
+
+/*********************************************************************************************************************************************************************************************************************************
+*** Serialization stuff for the old VM
+*********************************************************************************************************************************************************************************************************************************/
+#define VVMSer_batchStartExp(...)
+#define VVMSer_batchEndExp(...)
+
+#define VVMSer_chunkStartExp(...)
+#define VVMSer_chunkEndExp(...)
+
+#define VVMSer_insStartExp(...)
+#define VVMSer_insEndDecodeExp(...)
+#define VVMSer_insEndExp(...)
+
+#ifdef VVM_SERIALIZE_NO_WRITE
+#define VVMSer_chunkStart(...)
+#define VVMSer_chunkEnd(...)
+#define VVMSer_insStart(...)
+#define VVMSer_insEnd(...)
+#else
+#define VVMSer_chunkStart(Context, ChunkIdx, BatchIdx)                      \
+	FVectorVMSerializeChunk *SerializeChunk = nullptr;                      \
+	int VVMSerNumInstructionsThisChunk = 0;                                 \
+	unsigned char *VVMSerStartCtxCode = (unsigned char*)Context.Code;       \
+	if (SerializeState)                                                     \
+	{                                                                       \
+		SerializeChunk = SerializeState->Chunks + ChunkIdx;                 \
+		SerializeChunk->ChunkIdx = ChunkIdx;                                \
+		SerializeChunk->BatchIdx = BatchIdx;                                \
+		SerializeChunk->ExecIdx  = BatchIdx;                                \
+		SerializeChunk->NumInstances = NumInstancesThisChunk;               \
+		SerializeChunk->StartThreadID = FPlatformTLS::GetCurrentThreadId(); \
+		SerializeChunk->StartClock = FPlatformTime::Cycles64();             \
+	}
+
+#define VVMSer_chunkEnd(SerializeState, ...)																            \
+	if (SerializeChunk)																						            \
+	{																										            \
+		FPlatformAtomics::InterlockedOr(&SerializeState->ChunkComplete, (1ULL << (uint64)SerializeChunk->ChunkIdx));    \
+		SerializeChunk->EndThreadID = FPlatformTLS::GetCurrentThreadId();										        \
+		SerializeChunk->EndClock = FPlatformTime::Cycles64();												            \
+	}
+
+#define VVMSer_insStart(Context, ...)										\
+	uint64 VVMSerStartCycles = FPlatformTime::Cycles64();					\
+	unsigned char *VVMSerCtxStartInsCode = (unsigned char*)Context.Code;	\
+
+#define VVMSer_insEnd(Context, OpStart, NumOps, ...)																																												\
+	if (SerializeState && Op != EVectorVMOp::done)                                                                                                                                                                                  \
+	{																																												                                                \
+		uint64 VVMSerEndExecCycles = FPlatformTime::Cycles64();																																										\
+		VVMSer_serializeInstruction(SerializeState, Context, VVMSerNumInstructionsThisChunk++, StartInstance, NumInstancesThisChunk, ChunkIdx, OpStart, NumOps, VVMSerStartCycles, VVMSerEndExecCycles, VVMSerEndExecCycles);		\
+		uint64 VVMSerEndSerializeCycles = FPlatformTime::Cycles64();																																								\
+		SerializeState->SerializeDt += VVMSerEndSerializeCycles - VVMSerEndExecCycles;																																				\
+		SerializeChunk->InsExecTime += VVMSerEndSerializeCycles - VVMSerEndExecCycles;																																				\
+	}
+
+
+
+
+static void VVMSer_serializeInstruction(FVectorVMSerializeState *SerializeState, FVectorVMContext &Context, 
+										int InstructionIdx, int StartInstance, int NumInstancesThisChunk, int GlobalChunkIdx, 
+										int OpStart, int NumOps, uint64 StartInsCycles, uint64 EndDecodeCycles, uint64 EndInsCycles)
+{
+	FVectorVMSerializeInstruction *Ins = VVMSerGetNextInstruction(SerializeState, InstructionIdx, GlobalChunkIdx, 1, OpStart, NumOps, EndInsCycles - StartInsCycles, EndDecodeCycles - StartInsCycles);
+	if (Ins)
+	{
+		for (uint32 i = 0; i < SerializeState->NumTempRegisters; ++i)
+		{
+			unsigned char *TempReg = (unsigned char *)Context.GetTempRegister(i);
+			FMemory::Memcpy(Ins->TempRegisters + SerializeState->NumInstances * i + StartInstance, TempReg, sizeof(uint32) * NumInstancesThisChunk);
+		}
+	}
+}
+#endif //VVM_SERIALIZE_NO_WRITE
+
+#endif //NIAGARA_EXP_VM
+
+/*********************************************************************************************************************************************************************************************************************************
+*** Stuff shared between the new and old VM
+*********************************************************************************************************************************************************************************************************************************/
+
+static FVectorVMSerializeInstruction *VVMSerGetNextInstruction(FVectorVMSerializeState *SerializeState, int InstructionIdx, int GlobalChunkIdx, int BatchUseCount, int OpStart, int NumOps, uint64 Dt, uint64 DtDecode)
+{
+	FVectorVMSerializeInstruction *Ins = nullptr;
+	if (GlobalChunkIdx == 0 && BatchUseCount == 1)
+	{
+		check(InstructionIdx == SerializeState->NumInstructions);
+		if (SerializeState->NumInstructions + 1 >= SerializeState->NumInstructionsAllocated)
+		{
+			if (SerializeState->NumInstructionsAllocated == 0)
+			{
+				SerializeState->NumInstructionsAllocated = 1024;
+				check(SerializeState->Instructions == nullptr);
+			}
+			else
+			{
+				SerializeState->NumInstructionsAllocated <<= 1;
+			}
+			FVectorVMSerializeInstruction *NewInstructions = (FVectorVMSerializeInstruction *)SerializeState->ReallocFn(SerializeState->Instructions, sizeof(FVectorVMSerializeInstruction) * SerializeState->NumInstructionsAllocated, __FILE__, __LINE__);
+			if (NewInstructions == nullptr)
+			{
+				VectorVMSerializeSetError(SerializeState, VVMSerErr_OutOfMemory | VVMSerErr_Instruction | VVMSerErr_Fatal);
+				return nullptr;
+			}
+			else
+			{
+				SerializeState->Instructions = NewInstructions;
+			}
+		}
+		Ins = SerializeState->Instructions + SerializeState->NumInstructions++;
+
+		//we only store this information on the first chunk
+		Ins->OpStart       = OpStart;
+		Ins->NumOps        = NumOps;
+		Ins->Dt            = Dt;
+		Ins->DtDecode      = DtDecode;
+		Ins->TempRegisters = (uint32 *)SerializeState->ReallocFn(nullptr, sizeof(uint32) * SerializeState->NumInstances * SerializeState->NumTempRegisters, __FILE__, __LINE__); //alloc for *ALL* chunks
+		if (Ins->TempRegisters == nullptr)
+		{
+			VectorVMSerializeSetError(SerializeState, VVMSerErr_OutOfMemory | VVMSerErr_Instruction | VVMSerErr_Fatal);
+			return nullptr;
+		}
+		Ins->TempRegisterFlags = (unsigned char *)SerializeState->ReallocFn(nullptr, SerializeState->NumTempRegFlags, __FILE__, __LINE__);
+		if (Ins->TempRegisterFlags == nullptr)
+		{
+			VectorVMSerializeSetError(SerializeState, VVMSerErr_OutOfMemory | VVMSerErr_Instruction | VVMSerErr_Fatal);
+			return nullptr;
+		}
+		FMemory::Memset(Ins->TempRegisterFlags, 0, SerializeState->NumTempRegFlags);
+	}
+	else
+	{
+		if ((SerializeState->ChunkComplete & 1) == 0)
+		{
+			//we can't serialize the instructions until the first chunk is done since it's the one that allocates the memory.
+			//if this ever becomes a bottleneck (who cares for serialization!?), but if it ever does I can the change the 
+			//memory allocation such that each chunk's temp registers for each instruction gets grouped in per-chunk allocations
+			//then the debugger can merge them all into a single array on load... but this is simpler and I doubt will become an issue
+			uint32 SanityCount = 0;
+			while ((SerializeState->ChunkComplete & 1) == 0 && SanityCount++ < (1ULL << 31))
+			{
+				FPlatformProcess::Yield();
+			}
+			check(SanityCount < (1ULL << 31) - 1);
+		}
+		check(InstructionIdx < (int)SerializeState->NumInstructions);
+		Ins = SerializeState->Instructions + InstructionIdx;
+	}
+	return Ins;
+}
+
+#ifndef VVM_SERIALIZE_NO_IO
+VECTORVM_API void SerializeVectorVMWriteToFile(FVectorVMSerializeState *SerializeState, uint8 WhichStateWritten, const wchar_t *Filename)
+{
+	IPlatformFile &PlatformFile = FPlatformFileManager::Get().GetPlatformFile();
+	IFileHandle *File           = PlatformFile.OpenWrite((TCHAR *)Filename);
+	if (File)
+	{
+		File->Write(&WhichStateWritten, 1);		//1 = Exp, 2 = UE.  OR for both (1|2 = 3)
+		File->Write((uint8 *)&SerializeState->NumInstances     , sizeof(uint32));
+		File->Write((uint8 *)&SerializeState->Flags            , sizeof(uint32));
+		File->Write((uint8 *)&SerializeState->ExecDt           , sizeof(uint64));
+		File->Write((uint8 *)&SerializeState->SerializeDt      , sizeof(uint64));
+		File->Write((uint8 *)&SerializeState->NumTempRegisters , sizeof(uint32));
+		File->Write((uint8 *)&SerializeState->NumConstBuffers  , sizeof(uint32));
+		File->Write((uint8 *)&SerializeState->NumBytecodeBytes , sizeof(uint32));
+		File->Write((uint8 *)&SerializeState->NumInstructions  , sizeof(uint32));
+		File->Write((uint8 *)&SerializeState->NumDataSets      , sizeof(uint32));
+		File->Write((uint8 *)&SerializeState->NumChunks        , sizeof(uint32));
+		File->Write((uint8 *)&SerializeState->NumExternalData  , sizeof(uint32));
+		File->Write((uint8 *)&SerializeState->MaxExtFnRegisters, sizeof(uint32));
+		File->Write((uint8 *)&SerializeState->MaxExtFnUsed     , sizeof(uint32));
+
+		
+		//write the datasets
+		for (uint32 i = 0; i < SerializeState->NumDataSets; ++i)
+		{
+			FVectorVMSerializeDataSet *DataSet = SerializeState->DataSets + i;
+			uint32 IOFlag = (DataSet->InputBuffers ? 1 : 0) + (DataSet->OutputBuffers ? 2 : 0);
+			File->Write((uint8 *)&IOFlag                            , sizeof(uint32));
+			File->Write((uint8 *)&DataSet->InputOffset[0]           , sizeof(uint32));
+			File->Write((uint8 *)&DataSet->InputOffset[1]           , sizeof(uint32));
+			File->Write((uint8 *)&DataSet->InputOffset[2]           , sizeof(uint32));
+			File->Write((uint8 *)&DataSet->OutputOffset[0]          , sizeof(uint32));
+			File->Write((uint8 *)&DataSet->OutputOffset[1]          , sizeof(uint32));
+			File->Write((uint8 *)&DataSet->OutputOffset[2]          , sizeof(uint32));
+
+			File->Write((uint8 *)&DataSet->InputInstanceOffset      , sizeof(int32));
+			File->Write((uint8 *)&DataSet->InputDataSetAccessIndex  , sizeof(int32));
+			File->Write((uint8 *)&DataSet->InputIDAcquireTag        , sizeof(int32));
+			File->Write((uint8 *)&DataSet->InputNumFreeIDs          , sizeof(int32));
+			File->Write((uint8 *)&DataSet->InputMaxUsedIDs          , sizeof(int32));
+			File->Write((uint8 *)&DataSet->InputNumSpawnedIDs       , sizeof(int32));
+
+			File->Write((uint8 *)&DataSet->InputIDTableNum          , sizeof(int32));
+			File->Write((uint8 *)&DataSet->InputFreeIDTableNum      , sizeof(int32));
+			File->Write((uint8 *)&DataSet->InputSpawnedIDTableNum   , sizeof(int32));
+
+			File->Write((uint8 *)&DataSet->OutputInstanceOffset     , sizeof(int32));
+			File->Write((uint8 *)&DataSet->OutputDataSetAccessIndex , sizeof(int32));
+			File->Write((uint8 *)&DataSet->OutputIDAcquireTag       , sizeof(int32));
+			File->Write((uint8 *)&DataSet->OutputNumFreeIDs         , sizeof(int32));
+			File->Write((uint8 *)&DataSet->OutputMaxUsedIDs         , sizeof(int32));
+			File->Write((uint8 *)&DataSet->OutputNumSpawnedIDs      , sizeof(int32));
+
+			File->Write((uint8 *)&DataSet->OutputIDTableNum         , sizeof(int32));
+			File->Write((uint8 *)&DataSet->OutputFreeIDTableNum     , sizeof(int32));
+			File->Write((uint8 *)&DataSet->OutputSpawnedIDTableNum  , sizeof(int32));
+
+			uint32 TotalNumInstances = SerializeState->NumInstances + DataSet->InputInstanceOffset; //the VM could run 100 instances, but the dataset could have 150, and we start at 50, so just write all instances in the data set
+			//write input bufers
+			if (DataSet->InputBuffers)
+			{
+				uint32 NumBuffers = DataSet->InputOffset[2];
+				File->Write((uint8 *)DataSet->InputBuffers          , sizeof(uint32) * NumBuffers * TotalNumInstances);
+			}
+			if (DataSet->OutputBuffers)
+			{
+				uint32 NumBuffers = DataSet->OutputOffset[2];
+				File->Write((uint8 *)DataSet->OutputBuffers         , sizeof(uint32) * NumBuffers * TotalNumInstances);
+			}
+			//input tables
+			if (DataSet->InputIDTableNum > 0)
+			{
+				File->Write((uint8 *)DataSet->InputIDTable          , sizeof(int32) * DataSet->InputIDTableNum);
+			}
+			if (DataSet->InputFreeIDTableNum > 0)
+			{
+				File->Write((uint8 *)DataSet->InputFreeIDTable      , sizeof(int32) * DataSet->InputFreeIDTableNum);
+			}
+			if (DataSet->InputSpawnedIDTableNum > 0)
+			{
+				File->Write((uint8 *)DataSet->InputSpawnedIDTable   , sizeof(int32) * DataSet->InputSpawnedIDTableNum);
+			}
+			//output tables
+			if (DataSet->OutputIDTableNum > 0)
+			{
+				File->Write((uint8 *)DataSet->OutputIDTable         , sizeof(int32) * DataSet->OutputIDTableNum);
+			}
+			if (DataSet->OutputFreeIDTableNum > 0)
+			{
+				File->Write((uint8 *)DataSet->OutputFreeIDTable     , sizeof(int32) * DataSet->OutputFreeIDTableNum);
+			}
+			if (DataSet->OutputSpawnedIDTableNum > 0)
+			{
+				File->Write((uint8 *)DataSet->OutputSpawnedIDTable  , sizeof(int32) * DataSet->OutputSpawnedIDTableNum);
+			}
+		}
+
+		//write the chunk info
+		for (uint32 i = 0; i < SerializeState->NumChunks; ++i)
+		{
+			FVectorVMSerializeChunk *Chunk = SerializeState->Chunks + i;
+			File->Write((uint8 *)&Chunk->ChunkIdx		, sizeof(uint32));
+			File->Write((uint8 *)&Chunk->BatchIdx		, sizeof(uint32));
+			File->Write((uint8 *)&Chunk->NumInstances	, sizeof(uint32));
+			File->Write((uint8 *)&Chunk->StartThreadID	, sizeof(uint32));
+			File->Write((uint8 *)&Chunk->EndThreadID	, sizeof(uint32));
+			File->Write((uint8 *)&Chunk->StartClock	    , sizeof(uint64));
+			File->Write((uint8 *)&Chunk->EndClock		, sizeof(uint64));
+		}
+
+		//write const data
+		File->Write((uint8 *)SerializeState->PreExecConstData  , sizeof(uint32) * SerializeState->NumConstBuffers);
+		File->Write((uint8 *)SerializeState->PostExecConstData , sizeof(uint32) * SerializeState->NumConstBuffers);
+
+		//write bytecode
+		if (SerializeState->NumBytecodeBytes > 0)
+		{
+			File->Write((uint8 *)SerializeState->Bytecode, SerializeState->NumBytecodeBytes);
+		}
+
+		//write instructions
+		for (uint32 i = 0; i < SerializeState->NumInstructions; ++i)
+		{
+			FVectorVMSerializeInstruction *Ins = SerializeState->Instructions + i;
+			File->Write((uint8 *)&Ins->OpStart      , sizeof(uint32));
+			File->Write((uint8 *)&Ins->NumOps       , sizeof(uint32));
+			File->Write((uint8 *)&Ins->Dt           , sizeof(uint64));
+			File->Write((uint8 *)&Ins->DtDecode     , sizeof(uint64));
+			File->Write((uint8 *)Ins->TempRegisters , sizeof(uint32) * SerializeState->NumInstances * SerializeState->NumTempRegisters);
+			if (SerializeState->NumBytecodeBytes)
+			{
+				File->Write((uint8 *)Ins->TempRegisterFlags, SerializeState->NumTempRegisters);
+			}
+		}
+
+		//write external function data
+		for (uint32 i = 0; i < SerializeState->NumExternalData; ++i)
+		{
+			FVectorVMSerializeExternalData *ExtData = SerializeState->ExternalData + i;
+			File->Write((uint8 *)&ExtData->NameLen    , sizeof(uint16));
+			File->Write((uint8 *)ExtData->Name        , sizeof(uint16) * ExtData->NameLen);
+			File->Write((uint8 *)&ExtData->NumInputs  , sizeof(uint16));
+			File->Write((uint8 *)&ExtData->NumOutputs , sizeof(uint16));
+		}
+		delete File;
+	}
+}
+
+#else
+
+VECTORVM_API void SerializeVectorVMWriteToFile(FVectorVMSerializeState *SerializeState, uint8 WhichStateWritten, const wchar_t *Filename)
+{
+
+}
+
+#endif //VVM_SERIALIZE_NO_IO
+
+
+static uint32 SerializeConstData(FVectorVMSerializeState *SerializeState, FVectorVMConstData *ConstData, int NumConstData, uint32 **OutConstData)
+{
+	int TotalNumConstDWords = 0;
+	for (int i = 0; i < NumConstData; ++i)
+	{
+		TotalNumConstDWords += ConstData[i].NumDWords;
+	}
+	*OutConstData = (uint32 *)SerializeState->ReallocFn(nullptr, sizeof(uint32) * TotalNumConstDWords, __FILE__, __LINE__);
+	if (*OutConstData == nullptr)
+	{
+		VectorVMSerializeSetError(SerializeState, VVMSerErr_OutOfMemory | VVMSerErr_ConstData | VVMSerErr_Fatal);
+		return 0;
+	}
+	uint32 *cptr = *OutConstData;
+	for (int i = 0; i < NumConstData; ++i)
+	{
+		FMemory::Memcpy(cptr, ConstData[i].RegisterData, sizeof(uint32) * ConstData[i].NumDWords);
+		cptr += ConstData[i].NumDWords;
+	}
+	return TotalNumConstDWords;
+}
+
+VECTORVM_API uint32 SerializeVectorVMInputDataSets(FVectorVMSerializeState *SerializeState, TArrayView<FDataSetMeta> InDataSets, FVectorVMConstData *InConstData, int NumConstData)
+{
+	SerializeState->NumDataSets = InDataSets.Num();
+	SerializeState->DataSets = (FVectorVMSerializeDataSet *)SerializeState->ReallocFn(nullptr, sizeof(FVectorVMSerializeDataSet) * SerializeState->NumDataSets, __FILE__, __LINE__);
+	if (SerializeState->DataSets == nullptr)
+	{
+		return VectorVMSerializeSetError(SerializeState, VVMSerErr_OutOfMemory | VVMSerErr_InputDataSets | VVMSerErr_Fatal);
+	}
+	FMemory::Memset(SerializeState->DataSets, 0, sizeof(FVectorVMSerializeDataSet) * SerializeState->NumDataSets);
+	for (uint32 i = 0; i < SerializeState->NumDataSets; ++i)
+	{
+		FVectorVMSerializeDataSet *DsDst = SerializeState->DataSets + i;
+		FDataSetMeta *DsSrc = &InDataSets[i];
+		DsDst->InputInstanceOffset      = DsSrc->InstanceOffset;
+		DsDst->InputDataSetAccessIndex  = DsSrc->DataSetAccessIndex;
+		DsDst->InputIDAcquireTag        = DsSrc->IDAcquireTag;
+		DsDst->InputNumFreeIDs          = *DsSrc->NumFreeIDs;
+		DsDst->InputMaxUsedIDs          = *DsSrc->MaxUsedID;
+		DsDst->InputNumSpawnedIDs       = *DsSrc->NumSpawnedIDs;
+		
+		//we serialize the whole dataset... even if there's more instances than we use just so the offsets are correct when we sim
+		uint32 TotalNumInstances        = SerializeState->NumInstances + DsDst->InputInstanceOffset;
+		uint32 **SrcInputBuffers        = (uint32 **)DsSrc->InputRegisters.GetData();
+		if (SrcInputBuffers)
+		{
+			uint32 TotalNumBuffers = DsSrc->InputRegisterTypeOffsets[2];
+			DsDst->InputBuffers = (uint32 *)SerializeState->ReallocFn(nullptr, sizeof(uint32) * TotalNumBuffers * TotalNumInstances, __FILE__, __LINE__);
+			if (DsDst->InputBuffers == nullptr)
+			{
+				return VectorVMSerializeSetError(SerializeState, VVMSerErr_OutOfMemory | VVMSerErr_InputDataSets | VVMSerErr_Fatal);
+			}
+			for (uint32 j = 0; j < TotalNumBuffers; ++j)
+			{
+				FMemory::Memcpy(DsDst->InputBuffers + j * TotalNumInstances, SrcInputBuffers[j], sizeof(uint32) * TotalNumInstances);
+			}
+			DsDst->InputOffset[0] = DsSrc->InputRegisterTypeOffsets[0];
+			DsDst->InputOffset[1] = DsSrc->InputRegisterTypeOffsets[1];
+			DsDst->InputOffset[2] = DsSrc->InputRegisterTypeOffsets[2];
+		}
+		else
+		{
+			DsDst->InputOffset[0] = 0;
+			DsDst->InputOffset[1] = 0;
+			DsDst->InputOffset[2] = 0;
+			DsDst->InputBuffers   = nullptr;
+		}
+		
+		//IDTable
+		if (DsSrc->IDTable && DsSrc->IDTable->Num() > 0)
+		{
+			DsDst->InputIDTable = (int32 *)SerializeState->ReallocFn(nullptr, sizeof(int) * DsSrc->IDTable->Num(), __FILE__, __LINE__);
+			if (DsDst->InputIDTable == nullptr)
+			{
+				return VectorVMSerializeSetError(SerializeState, VVMSerErr_OutOfMemory | VVMSerErr_InputDataSets | VVMSerErr_Fatal);
+			}
+			FMemory::Memcpy(DsDst->InputIDTable, DsSrc->IDTable->GetData(), sizeof(int32) * DsSrc->IDTable->Num());
+			DsDst->InputIDTableNum = DsSrc->IDTable->Num();
+		}
+		else
+		{
+			DsDst->InputIDTable    = nullptr;
+			DsDst->InputIDTableNum = 0;
+		}
+
+		//FreeIDTable
+		if (DsSrc->FreeIDTable && DsSrc->FreeIDTable->Num() > 0)
+		{
+			DsDst->InputFreeIDTable = (int32 *)SerializeState->ReallocFn(nullptr, sizeof(int) * DsSrc->FreeIDTable->Num(), __FILE__, __LINE__);
+			if (DsDst->InputFreeIDTable == nullptr)
+			{
+				return VectorVMSerializeSetError(SerializeState, VVMSerErr_OutOfMemory | VVMSerErr_InputDataSets | VVMSerErr_Fatal);
+			}
+			FMemory::Memcpy(DsDst->InputFreeIDTable, DsSrc->FreeIDTable->GetData(), sizeof(int32) * DsSrc->FreeIDTable->Num());
+			DsDst->InputFreeIDTableNum = DsSrc->FreeIDTable->Num();
+		}
+		else
+		{
+			DsDst->InputFreeIDTable    = nullptr;
+			DsDst->InputFreeIDTableNum = 0;
+		}
+
+		//SpawnedIDTable
+		if (DsSrc->SpawnedIDsTable && DsSrc->SpawnedIDsTable->Num() > 0)
+		{
+			DsDst->InputSpawnedIDTable = (int32 *)SerializeState->ReallocFn(nullptr, sizeof(int) * DsSrc->SpawnedIDsTable->Num(), __FILE__, __LINE__);
+			if (DsDst->InputSpawnedIDTable == nullptr)
+			{
+				return VectorVMSerializeSetError(SerializeState, VVMSerErr_OutOfMemory | VVMSerErr_InputDataSets | VVMSerErr_Fatal);
+			}
+			FMemory::Memcpy(DsDst->InputSpawnedIDTable, DsSrc->SpawnedIDsTable->GetData(), sizeof(int32) * DsSrc->SpawnedIDsTable->Num());
+			DsDst->InputSpawnedIDTableNum = DsSrc->SpawnedIDsTable->Num();
+		}
+		else
+		{
+			DsDst->InputSpawnedIDTable    = nullptr;
+			DsDst->InputSpawnedIDTableNum = 0;
+		}
+	}
+	SerializeState->NumConstBuffers = SerializeConstData(SerializeState, InConstData, NumConstData, &SerializeState->PreExecConstData);
+	return SerializeState->Error.Flags;
+}
+
+VECTORVM_API uint32 SerializeVectorVMOutputDataSets(FVectorVMSerializeState *SerializeState, TArrayView<FDataSetMeta> InDataSets, FVectorVMConstData *InConstData, int InNumConstData)
+{
+	if (SerializeState->NumDataSets == InDataSets.Num()) //must have already serialized the input data sets
+	{
+		for (int i = 0; i < InDataSets.Num(); ++i)
+		{
+			FVectorVMSerializeDataSet *DsDst = SerializeState->DataSets + i;
+			FDataSetMeta *DsSrc = &InDataSets[i];
+			DsDst->OutputOffset[0]			= DsSrc->OutputRegisterTypeOffsets[0];
+			DsDst->OutputOffset[1]			= DsSrc->OutputRegisterTypeOffsets[1];
+			DsDst->OutputOffset[2]			= DsSrc->OutputRegisterTypeOffsets[2];
+			DsDst->OutputInstanceOffset		= DsSrc->InstanceOffset;
+			DsDst->OutputDataSetAccessIndex	= DsSrc->DataSetAccessIndex;
+			DsDst->OutputIDAcquireTag		= DsSrc->IDAcquireTag;
+			DsDst->OutputNumFreeIDs         = *DsSrc->NumFreeIDs;
+			DsDst->OutputMaxUsedIDs         = *DsSrc->MaxUsedID;
+			DsDst->OutputNumSpawnedIDs      = *DsSrc->NumSpawnedIDs;
+
+			uint32 TotalNumInstances = SerializeState->NumInstances + DsDst->OutputInstanceOffset;
+			uint32 **SrcOutputBuffers = (uint32 **)DsSrc->OutputRegisters.GetData();
+			if (SrcOutputBuffers)
+			{
+				uint32 TotalNumBuffers = DsSrc->OutputRegisterTypeOffsets[2];
+				DsDst->OutputBuffers = (uint32 *)SerializeState->ReallocFn(nullptr, sizeof(uint32) * TotalNumBuffers * TotalNumInstances, __FILE__, __LINE__);
+				if (DsDst->OutputBuffers == nullptr)
+				{
+					return VectorVMSerializeSetError(SerializeState, VVMSerErr_OutOfMemory | VVMSerErr_OutputDataSets | VVMSerErr_Fatal);
+				}
+				for (uint32 j = 0; j < TotalNumBuffers; ++j)
+				{
+					FMemory::Memcpy(DsDst->OutputBuffers + j * TotalNumInstances, SrcOutputBuffers[j], sizeof(uint32) * TotalNumInstances);
+				}
+			}
+			else
+			{
+				DsDst->OutputBuffers = nullptr;
+			}
+
+			//IDTable
+			if (DsSrc->IDTable && DsSrc->IDTable->Num() > 0)
+			{
+				DsDst->OutputIDTable = (int32 *)SerializeState->ReallocFn(nullptr, sizeof(int) * DsSrc->IDTable->Num(), __FILE__, __LINE__);
+				if (DsDst->OutputIDTable == nullptr)
+				{
+					return VectorVMSerializeSetError(SerializeState, VVMSerErr_OutOfMemory | VVMSerErr_OutputDataSets | VVMSerErr_Fatal);
+				}
+				FMemory::Memcpy(DsDst->OutputIDTable, DsSrc->IDTable->GetData(), sizeof(int32) * DsSrc->IDTable->Num());
+				DsDst->OutputIDTableNum = DsSrc->IDTable->Num();
+			}
+			else
+			{
+				DsDst->OutputIDTable    = nullptr;
+				DsDst->OutputIDTableNum = 0;
+			}
+
+			//FreeIDTable
+			if (DsSrc->FreeIDTable && DsSrc->FreeIDTable->Num() > 0)
+			{
+				DsDst->OutputFreeIDTable = (int32 *)SerializeState->ReallocFn(nullptr, sizeof(int) * DsSrc->FreeIDTable->Num(), __FILE__, __LINE__);
+				if (DsDst->OutputFreeIDTable == nullptr)
+				{
+					return VectorVMSerializeSetError(SerializeState, VVMSerErr_OutOfMemory | VVMSerErr_OutputDataSets | VVMSerErr_Fatal);
+				}
+				FMemory::Memcpy(DsDst->OutputFreeIDTable, DsSrc->FreeIDTable->GetData(), sizeof(int32) * DsSrc->FreeIDTable->Num());
+				DsDst->OutputFreeIDTableNum = DsSrc->FreeIDTable->Num();
+			}
+			else
+			{
+				DsDst->OutputFreeIDTable    = nullptr;
+				DsDst->OutputFreeIDTableNum = 0;
+			}
+
+			//SpawnedIDTable
+			if (DsSrc->SpawnedIDsTable && DsSrc->SpawnedIDsTable->Num() > 0)
+			{
+				DsDst->OutputSpawnedIDTable = (int32 *)SerializeState->ReallocFn(nullptr, sizeof(int) * DsSrc->SpawnedIDsTable->Num(), __FILE__, __LINE__);
+				if (DsDst->OutputSpawnedIDTable == nullptr)
+				{
+					return VectorVMSerializeSetError(SerializeState, VVMSerErr_OutOfMemory | VVMSerErr_OutputDataSets | VVMSerErr_Fatal);
+				}
+				FMemory::Memcpy(DsDst->OutputSpawnedIDTable, DsSrc->SpawnedIDsTable->GetData(), sizeof(int32) * DsSrc->SpawnedIDsTable->Num());
+				DsDst->OutputSpawnedIDTableNum = DsSrc->SpawnedIDsTable->Num();
+			}
+			else
+			{
+				DsDst->OutputSpawnedIDTable    = nullptr;
+				DsDst->OutputSpawnedIDTableNum = 0;
+			}
+
+		}
+	}
+	else
+	{
+		VectorVMSerializeSetError(SerializeState, VVMSerErr_OutputDataSets);
+		for (uint32 i = 0; i < SerializeState->NumDataSets; ++i)
+		{
+			FVectorVMSerializeDataSet *DataSet = SerializeState->DataSets + i;
+			DataSet->OutputOffset[0] = 0;
+			DataSet->OutputOffset[1] = 0;
+			DataSet->OutputOffset[2] = 0;
+			DataSet->OutputBuffers = nullptr;
+			DataSet->OutputIDTable = nullptr;
+			DataSet->OutputIDTableNum = 0;
+			DataSet->OutputFreeIDTable = nullptr;
+			DataSet->OutputFreeIDTableNum = 0;
+			DataSet->OutputSpawnedIDTable = nullptr;
+			DataSet->OutputSpawnedIDTableNum = 0;
+		}
+	}
+	if (SerializeState->Error.Flags == 0)
+	{
+		SerializeConstData(SerializeState, InConstData, InNumConstData, &SerializeState->PostExecConstData);
+	}
+	return SerializeState->Error.Flags;
+}
+
+#else //VVM_INCLUDE_SERIALIZATION
+
+#define VVMSer_batchStart(...)
+#define VVMSer_batchEnd(...)
+
+#define VVMSer_chunkStart(...)
+#define VVMSer_chunkEnd(...)
+
+#define VVMSer_insStart(...)
+#define VVMSer_insEndDecode(...)
+#define VVMSer_insEnd(...)
+
+#define VVMSer_batchStartExp(...)
+#define VVMSer_batchEndExp(...)
+
+#define VVMSer_chunkStartExp(...)
+#define VVMSer_chunkEndExp(...)
+
+#define VVMSer_insStartExp(...)
+#define VVMSer_insEndDecodeExp(...)
+#define VVMSer_insEndExp(...)
+
+#define VVMSer_initSerializationState(...)
+#define FreeVectorVMSerializeState(...)
+
+VECTORVM_API uint64 VVMSer_cmpStates(FVectorVMSerializeState *S0, FVectorVMSerializeState *S1, uint32 Flags)
+{
+	return 0;
+}
+
+VECTORVM_API uint32 SerializeVectorVMInputDataSets(FVectorVMSerializeState *SerializeState, TArrayView<FDataSetMeta> DataSets, FVectorVMConstData *ConstData, int NumConstData)
+{
+	return 0;
+}
+
+VECTORVM_API uint32 SerializeVectorVMOutputDataSets(FVectorVMSerializeState *SerializeState, TArrayView<FDataSetMeta> DataSets, FVectorVMConstData *ConstData, int NumConstData)
+{
+	return 0;
+}
+
+#endif //VVM_INCLUDE_SERIALIZATION
diff --git a/Engine/Source/Runtime/VectorVM/Public/VectorVM.h b/Engine/Source/Runtime/VectorVM/Public/VectorVM.h
--- a/Engine/Source/Runtime/VectorVM/Public/VectorVM.h
+++ b/Engine/Source/Runtime/VectorVM/Public/VectorVM.h
@@ -156,8 +156,40 @@ enum class EVectorVMOp : uint8
 	//acquires a new ID from the free list.
 	acquire_id,
 
+	//experimental VM only
+	fused_input1_1, /* 95    op has 1 input operand	   - binary 1          ie: register 0 is an input          (ORDER IS CRUCIAL) */
+	fused_input2_1, /* 96    op has 2 input operands   - binary 01         ie: register 0 is an input          (ORDER IS CRUCIAL) */
+	fused_input2_2, /* 97    op has 2 input operands   - binary 10         ie: register 1 is an input          (ORDER IS CRUCIAL) */
+	fused_input2_3, /* 97    op has 2 input operands   - binary 11         ie: register 1 and 2 are inputs     (ORDER IS CRUCIAL) */
+	fused_input3_1, /* 98    op has 3 input operands   - binary 001        ie: register 1 is an input          (ORDER IS CRUCIAL) */
+	fused_input3_2, /* 99    op has 3 input operands   - binary 010        ie: register 2 is an input          (ORDER IS CRUCIAL) */
+	fused_input3_4, /* 101   op has 3 input operands   - binary 100        ie: register 3 is an input          (ORDER IS CRUCIAL) */
+	fused_input3_3, /* 100   op has 3 input operands   - binary 011        ie: register 1 and 2 are inputs     (ORDER IS CRUCIAL) */
+	fused_input3_5, /* 102   op has 3 input operands   - binary 101        ie: register 3 and 1 are inputs     (ORDER IS CRUCIAL) */
+	fused_input3_6, /* 103   op has 3 input operands   - binary 110        ie: register 2 and 2 are inputs     (ORDER IS CRUCIAL) */
+	fused_input3_7, /* 104   op has 3 input operands   - binary 111        ie: register 1, 2 and 3 are inputs  (ORDER IS CRUCIAL) */
+	copy_to_output, /* 105 */
+	output_batch2,  /* 106 */
+	output_batch3,  /* 107 */
+	output_batch4,  /* 108 */
+	output_batch7,  /* 108 */
+	output_batch8,  /* 108 */
+
 	NumOpcodes
 };
+	
+enum class EVectorVMOpCategory : uint8 {
+	Input,
+	Output,
+	Op,
+	ExtFnCall,
+	IndexGen,
+	ExecIndex,
+	RWBuffer,
+	Stat,
+	Fused,
+	Other
+};
 
 #if STATS
 struct FVMCycleCounter
@@ -215,14 +247,19 @@ struct FDataSetMeta
 	/** MaxID used in this execution. */
 	int32* MaxUsedID;
 
+	int32 *NumSpawnedIDs;
+
 	int32 IDAcquireTag;
 
 	//Temporary lock we're using for thread safety when writing to the FreeIDTable.
 	//TODO: A lock free algorithm is possible here. We can create a specialized lock free list and reuse the IDTable slots for FreeIndices as Next pointers for our LFL.
 	//This would also work well on the GPU. 
 	//UE-65856 for tracking this work.
-	FCriticalSection FreeTableLock;
 
+	//@NOTE (smcgrath): the lock/unlock functions can go with the ne VM, we don't need them
+#ifndef NIAGARA_EXP_VM
+	FCriticalSection FreeTableLock;
+#endif
 	FORCEINLINE void LockFreeTable();
 	FORCEINLINE void UnlockFreeTable();
 
@@ -254,7 +291,7 @@ struct FDataSetMeta
 		IDAcquireTag = INDEX_NONE;
 	}
 
-	FORCEINLINE void Init(const TArrayView<uint8 const* RESTRICT const>& InInputRegisters, const TArrayView<uint8 const* RESTRICT const>& InOutputRegisters, int32 InInstanceOffset, TArray<int32>* InIDTable, TArray<int32>* InFreeIDTable, int32* InNumFreeIDs, int32* InMaxUsedID, int32 InIDAcquireTag, TArray<int32>* InSpawnedIDsTable)
+	FORCEINLINE void Init(const TArrayView<uint8 const* RESTRICT const>& InInputRegisters, const TArrayView<uint8 const* RESTRICT const>& InOutputRegisters, int32 InInstanceOffset, TArray<int32>* InIDTable, TArray<int32>* InFreeIDTable, int32* InNumFreeIDs, int32 *InNumSpawnedIDs, int32* InMaxUsedID, int32 InIDAcquireTag, TArray<int32>* InSpawnedIDsTable)
 	{
 		InputRegisters = InInputRegisters;
 		OutputRegisters = InOutputRegisters;
@@ -264,6 +301,7 @@ struct FDataSetMeta
 		IDTable = InIDTable;
 		FreeIDTable = InFreeIDTable;
 		NumFreeIDs = InNumFreeIDs;
+		NumSpawnedIDs = InNumSpawnedIDs;
 		MaxUsedID = InMaxUsedID;
 		IDAcquireTag = InIDAcquireTag;
 		SpawnedIDsTable = InSpawnedIDsTable;
@@ -277,6 +315,8 @@ private:
 	FDataSetMeta& operator=(const FDataSetMeta&) = delete;
 };
 
+#include "VectorVMExperimental.h"
+
 //Data the VM will keep on each dataset locally per thread which is then thread safely pushed to it's destination at the end of execution.
 struct FDataSetThreadLocalTempData
 {
@@ -481,6 +521,42 @@ public:
 };
 
 class FVectorVMExternalFunctionContext {
+#ifdef NIAGARA_EXP_VM
+public:
+	uint32 **                RegisterData;
+	uint16 *                 RawVecIndices; //undecoded, for compatbility with the previous VM
+	uint32 *                 RegInc;
+
+	int                      RegReadCount;
+	int                      NumRegisters;
+
+	int                      StartInstance;
+	int                      NumInstances;
+	int                      NumLoops;
+	int                      PerInstanceFnInstanceIdx;
+
+	void **                  UserPtrTable;
+	int                      NumUserPtrs;
+	
+	FRandomStream *          RandStream;
+	int32 *                  RandCounters;
+	TArrayView<FDataSetMeta> DataSets;
+
+	FORCEINLINE int32                                  GetStartInstance() const             { return StartInstance; }
+	FORCEINLINE int32                                  GetNumInstances() const              { return NumInstances; }
+	FORCEINLINE int32 *                                GetRandCounters()                    { return RandCounters; }
+	FORCEINLINE FRandomStream &                        GetRandStream()                      { return *RandStream; }
+	FORCEINLINE void *                                 GetUserPtrTable(int32 UserPtrIdx)    { check(UserPtrIdx < NumUserPtrs);  return UserPtrTable[UserPtrIdx]; }
+	template<uint32 InstancesPerOp> FORCEINLINE int32  GetNumLoops() const                  { static_assert(InstancesPerOp == 4); return NumLoops; };
+
+	FORCEINLINE float *GetNextRegister(int32 *OutAdvanceOffset, int32 *OutVecIndex) {
+		check(RegReadCount < NumRegisters);
+		*OutAdvanceOffset = RegInc[RegReadCount] & 1;
+		*OutVecIndex      = RawVecIndices[RegReadCount];
+		return (float *)RegisterData[RegReadCount++];
+	}
+
+#else
 	friend struct FNiagaraSystemScriptExecutionContext; //@NOTE(smcgrath): required for the PerInstanceFunctionHook() in the non-experimental version of VectorVM
 public:
 	FVectorVMExternalFunctionContext(FVectorVMContext *InVectorVMContext) : VectorVMContext(InVectorVMContext) { }
@@ -508,6 +584,7 @@ public:
 	FORCEINLINE int32 GetNumLoops() const { return VectorVMContext->GetNumLoops<InstancesPerOp>(); }
 private:
 	FVectorVMContext *VectorVMContext;
+#endif
 };
 
 namespace VectorVM
@@ -545,7 +622,7 @@ namespace VectorVM
 	/**
 	 * Execute VectorVM bytecode.
 	 */
-	VECTORVM_API void Exec(FVectorVMExecArgs& Args);
+	VECTORVM_API void Exec(FVectorVMExecArgs& Args, FVectorVMSerializeState *SerializeState);
 
 	VECTORVM_API void OptimizeByteCode(const uint8* ByteCode, TArray<uint8>& OptimizedCode, TArrayView<uint8> ExternalFunctionRegisterCounts);
 
@@ -562,7 +639,13 @@ namespace VectorVM
 		FUserPtrHandler(FVectorVMExternalFunctionContext& Context)
 		{
 #ifdef NIAGARA_EXP_VM
-			check(false);
+			int32 AdvanceOffset;
+			int32 RegIdx;
+			int32 *ConstPtr = (int32 *)Context.GetNextRegister(&AdvanceOffset, &RegIdx);
+			check(AdvanceOffset == 0); //must be constant
+			UserPtrIdx = *ConstPtr;
+			check(UserPtrIdx != INDEX_NONE);
+			Ptr = (T *)Context.GetUserPtrTable(UserPtrIdx);
 #else
 			const uint16 VariableOffset = Context.DecodeU16();
 			check(!(VariableOffset & VVM_EXT_FUNC_INPUT_LOC_BIT));
@@ -609,7 +692,7 @@ namespace VectorVM
 		void Init(FVectorVMExternalFunctionContext& Context)
 		{
 #ifdef NIAGARA_EXP_VM
-			check(false);
+			InputPtr = (T*)Context.GetNextRegister(&AdvanceOffset, &InputOffset) + Context.PerInstanceFnInstanceIdx;
 #else
 			InputOffset = Context.DecodeU16();
 
@@ -654,7 +737,7 @@ namespace VectorVM
 	public:
 #ifdef NIAGARA_EXP_VM
 		FORCEINLINE FExternalFuncRegisterHandler(FVectorVMExternalFunctionContext& Context) {
-			check(false);
+			Register = (T*)Context.GetNextRegister(&AdvanceOffset, &RegisterIndex) + Context.PerInstanceFnInstanceIdx;
 		}
 #else
 		FORCEINLINE FExternalFuncRegisterHandler(FVectorVMExternalFunctionContext& Context)
diff --git a/Engine/Source/Runtime/VectorVM/Public/VectorVMExperimental.h b/Engine/Source/Runtime/VectorVM/Public/VectorVMExperimental.h
new file mode 100644
--- /dev/null
+++ b/Engine/Source/Runtime/VectorVM/Public/VectorVMExperimental.h
@@ -0,0 +1,451 @@
+// Copyright Epic Games, Inc. All Rights Reserved.
+#pragma once
+
+//only to be included by VectorVM.h
+
+typedef void * (VectorVMReallocFn)              (void *Ptr, size_t NumBytes, const char *Filename, int LineNumber);
+typedef void   (VectorVMFreeFn)                 (void *Ptr, const char *Filename, int LineNumber);
+
+//prototypes for serialization that are required whether or not serialization is enabled
+#if defined(NIAGARA_EXP_VM) || defined(VVM_INCLUDE_SERIALIZATION)
+struct FVectorVMConstData
+{
+	void * RegisterData;
+	uint32 NumDWords;
+	uint32 StartingOffset;
+};
+
+typedef uint32 (VectorVMOptimizeErrorCallback)  (struct FVectorVMOptimizeContext *OptimizeContext, uint32 ErrorFlags);	//return new error flags
+typedef uint32 (VectorVMSerializeErrorCallback) (struct FVectorVMSerializeState *SerializeState, uint32 ErrorFlags);
+
+#endif //NIAGARA_EXP_VM || VVM_INCLUDE_SERIALIZATION
+
+//Serialization
+enum EVectorVMSerializeFlags
+{
+	VVMSer_SyncRandom        = 1 << 0,
+	VVMSer_SyncExtFns        = 1 << 1,
+	VVMSer_OptimizedBytecode = 1 << 2,
+};
+
+#ifdef VVM_INCLUDE_SERIALIZATION
+
+enum EVectorVMSerializeError
+{
+	VVMSerErr_OutOfMemory    = 1 << 0,
+	VVMSerErr_Init           = 1 << 1,
+	VVMSerErr_InputDataSets  = 1 << 2,
+	VVMSerErr_OutputDataSets = 1 << 3,
+	VVMSerErr_Instruction    = 1 << 4,
+	VVMSerErr_ConstData      = 1 << 5,
+
+	VVMSerErr_Fatal          = 1 << 31
+};
+
+struct FVectorVMSerializeInstruction
+{
+	uint32   OpStart;
+	uint32   NumOps;
+	uint64   Dt;
+	uint64   DtDecode;
+	uint32 * TempRegisters;
+	uint8 *  TempRegisterFlags;
+};
+
+struct FVectorVMSerializeExternalData
+{
+	wchar_t * Name;
+	uint16    NameLen;
+	uint16    NumInputs;
+	uint16    NumOutputs;
+};
+
+struct FVectorVMSerializeDataSet
+{
+	uint32 * InputBuffers;
+	uint32 * OutputBuffers;
+
+	uint32  InputOffset[3];	    //float, int, half (half must be 0)
+	uint32  OutputOffset[3];    //float, int, half (half must be 0)
+
+	int32   InputInstanceOffset;
+	int32   InputDataSetAccessIndex;
+	int32   InputIDAcquireTag;
+
+	int32   OutputInstanceOffset;
+	int32   OutputDataSetAccessIndex;
+	int32   OutputIDAcquireTag;
+
+	int32 * InputIDTable;
+	int32 * InputFreeIDTable;
+	int32 * InputSpawnedIDTable;
+
+	int32   InputIDTableNum;
+	int32   InputFreeIDTableNum;
+	int32   InputSpawnedIDTableNum;
+
+	int32   InputNumFreeIDs;
+	int32   InputMaxUsedIDs;
+	int32   InputNumSpawnedIDs;
+
+	int32 * OutputIDTable;
+	int32 * OutputFreeIDTable;
+	int32 * OutputSpawnedIDTable;
+
+	int32   OutputIDTableNum;
+	int32   OutputFreeIDTableNum;
+	int32   OutputSpawnedIDTableNum;
+
+	int32   OutputNumFreeIDs;
+	int32   OutputMaxUsedIDs;
+	int32   OutputNumSpawnedIDs;
+};
+
+struct FVectorVMSerializeChunk
+{
+	uint32 ChunkIdx;
+	uint32 BatchIdx;
+	uint32 ExecIdx;
+	uint32 StartInstance;
+	uint32 NumInstances;
+
+	uint32 StartThreadID;
+	uint32 EndThreadID;
+
+	uint64 StartClock;
+	uint64 EndClock;
+	uint64 InsExecTime;
+};
+
+struct FVectorVMSerializeState
+{
+	uint32                              NumInstances;
+	uint32                              NumTempRegisters;
+	uint32                              NumTempRegFlags;  //max of NumTempRegisters and Num Input Registers in each dataset
+	uint32                              NumConstBuffers;
+
+	uint32                              Flags;
+
+	FVectorVMSerializeInstruction *     Instructions;
+	uint32                              NumInstructions;
+	uint32                              NumInstructionsAllocated;
+
+	uint32                              NumExternalData;
+	FVectorVMSerializeExternalData *    ExternalData;
+	uint32								MaxExtFnRegisters;
+	uint32								MaxExtFnUsed;
+
+	uint64                              ExecDt;
+	uint64                              SerializeDt;
+
+	uint8 *                             TempRegFlags;
+	uint8 *                             Bytecode;
+	uint32                              NumBytecodeBytes;
+
+	FVectorVMSerializeDataSet *         DataSets;
+	uint32                              NumDataSets;
+	uint32 *                            PreExecConstData;
+	uint32 *                            PostExecConstData;
+
+	uint32                              NumChunks;
+	FVectorVMSerializeChunk *           Chunks;
+
+	struct FVectorVMOptimizeContext *   OptimizeCtx;
+
+	volatile int64                      ChunkComplete; //1 bit for each of the first 64 chunks
+
+	VectorVMReallocFn *                  ReallocFn;
+	VectorVMFreeFn *                     FreeFn;
+
+	struct {
+		uint32                           Flags;
+		uint32                           Line;
+		VectorVMSerializeErrorCallback * CallbackFn;
+	} Error;
+};
+#else // VVM_INCLUDE_SERIALIZATION
+
+struct FVectorVMSerializeState
+{
+	uint32              Flags;
+	VectorVMReallocFn * ReallocFn;
+	VectorVMFreeFn *    FreeFn;
+
+};
+
+#endif // VVM_INCLUDE_SERIALIZATION
+
+#ifdef NIAGARA_EXP_VM
+
+union FVecReg {
+	VectorRegister4f v;
+	VectorRegister4i i;
+};
+
+struct FVectorVMExtFunctionData
+{
+	const FVMExternalFunction *Function;
+	int32                      NumInputs;
+	int32                      NumOutputs;
+};
+
+//Optimization
+enum EVectorVMOptimizeFlags
+{
+	VVMOptFlag_SaveIntermediateState = 1 << 0,
+	VVMOptFlag_OmitStats             = 1 << 1
+};
+
+struct FVectorVMOptimizeInstruction
+{
+	EVectorVMOp         OpCode;
+	EVectorVMOpCategory OpCat;
+	uint32              PtrOffsetInOrigBytecode;
+	uint32              PtrOffsetInOptimizedBytecode;
+	int                 Index; //initial index.  Instructions are moved around and removed and dependency chains are created based on index, so we need to store this.
+	union
+	{
+		struct
+		{
+			uint16 DstRegPtrOffset;
+			uint16 DataSetIdx;
+			uint16 InputIdx;
+			uint16 FuseCount;
+			int    FirstInsInsertIdx;
+		} Input;
+		struct
+		{
+			uint16 RegPtrOffset;
+			uint16 DataSetIdx;
+			uint16 DstRegIdx;
+			int    CopyFromInputInsIdx;
+		} Output;
+		struct
+		{
+			uint32 RegPtrOffset;
+			uint16 NumInputs;
+			uint16 NumOutputs;
+			uint8  InputFuseBits;
+		} Op;
+		struct
+		{
+			uint32 RegPtrOffset;
+			uint16 DataSetIdx;
+		} IndexGen;
+		struct
+		{
+			uint32 RegPtrOffset;
+			uint16 ExtFnIdx;
+			uint16 NumInputs;
+			uint16 NumOutputs;
+		} ExtFnCall;
+		struct
+		{
+			uint32 RegPtrOffset;
+		} ExecIndex;
+		struct
+		{
+			uint32 RegPtrOffset;
+			uint16 DataSetIdx;
+		} RWBuffer;
+		struct
+		{
+			uint16 ID;
+		} Stat;
+		struct
+		{
+		} Other;
+	};
+};
+
+enum EVectorVMOptimizeError
+{
+	VVMOptErr_OutOfMemory        = 1 << 0,
+	VVMOptErr_Bytecode           = 1 << 1,
+	VVMOptErr_RegisterUsage      = 1 << 2,
+	VVMOptErr_ConstRemap         = 1 << 3,
+	VVMOptErr_Instructions       = 1 << 4,
+	VVMOptErr_InputFuseBuffer    = 1 << 5,
+	VVMOptErr_InstructionReOrder = 1 << 6,
+	VVMOptErr_SSARemap           = 1 << 7,
+	VVMOptErr_OptimizedBytecode  = 1 << 8,
+	VVMOptErr_ExternalFunction   = 1 << 9,
+
+	VVMOptErr_Fatal              = 1 << 31
+};
+
+struct FVectorVMOptimizeContext
+{
+	uint8 *                               OutputBytecode;
+	uint16 *                              ConstRemap[2];
+	FVectorVMExtFunctionData *            ExtFnTable;
+	uint32                                NumBytecodeBytes;
+	uint32                                NumOutputDataSets;
+	uint16                                NumConstsAlloced;    //upper bound to alloc
+	uint16                                NumConstsRemapped;
+	uint32                                NumTempRegisters;
+	uint32                                NumExtFns;
+	uint32                                MaxExtFnRegisters;
+	int32                                 MaxExtFnUsed;
+
+	struct
+	{
+		VectorVMReallocFn *               ReallocFn;
+		VectorVMFreeFn *                  FreeFn;
+	} Init;                                                    //Set this stuff when calling Optimize()
+
+	struct
+	{
+		uint32                            Flags;               //zero is good
+		uint32                            Line;
+		VectorVMOptimizeErrorCallback *   CallbackFn;          //set this to get a callback whenever there's an error
+	} Error;
+
+	struct
+	{
+		FVectorVMOptimizeInstruction *    Instructions;
+		uint16 *                          RegisterUsageBuffer;
+		uint16 *                          SSARegisterUsageBuffer;
+		int32 *                           InputRegisterFuseBuffer;
+		uint32                            NumBytecodeBytes;
+		uint32                            NumInstructions;
+		uint32                            NumRegistersUsed;
+	} Intermediate;                                           //these are freed and NULL after optimize() unless SaveIntermediateData is true when calling OptimizeVectorVMScript
+};
+
+//VectorVMState
+enum EVectorVMStateError
+{
+	VVMErr_InitOutOfMemory     = 1 << 0,
+	VVMErr_InitMemMismatch     = 1 << 1,
+	VVMErr_BatchMemory         = 1 << 2,
+	VVMErr_AssignInstances     = 1 << 3,
+
+	VVMErr_Fatal               = 1 << 31
+};
+
+struct FVectorVMExternalFnPerInstanceData
+{
+	class UNiagaraDataInterface *     DataInterface;
+	void *                            UserData;
+	uint16                            NumInputs;
+	uint16                            NumOutputs;
+};
+
+struct FVectorVMInitData
+{
+	struct FVectorVMState *					ExistingVectorVMState;
+	FVectorVMOptimizeContext *              OptimizeContext;
+	TArrayView<FDataSetMeta>                DataSets;
+	TArrayView<const FVMExternalFunction *> ExtFunctionTable;
+	
+	int                                     NumInstances;
+
+	int                                     NumConstData;
+	FVectorVMConstData *                    ConstData;
+
+	void **                                 UserPtrTable;
+	int                                     NumUserPtrTable;
+
+	VectorVMReallocFn *                     ReallocFn;
+	VectorVMFreeFn *                        FreeFn;
+};
+
+struct FVectorVMBatchState
+{
+	MS_ALIGN(16) FVecReg *    RegisterData GCC_ALIGN(16);
+	struct {
+		uint32 *              StartingOutputIdxPerDataSet;
+		uint32 *              NumOutputPerDataSet;
+		struct {
+			FVecReg **        RegData;
+			uint32 *          RegInc;
+		} ExtFnDecodedReg;
+		int32 *               RandCounters; //used for external functions only.
+	} ChunkLocalData;
+	void *                    MallocedMemPtr; //needed for alignment purposes
+
+	volatile int32			  CurrentlyExecuting;
+
+	int                       StartInstance;
+	int                       NumInstances;
+
+	VectorRegister4i          RandState[5]; //xorwor state for random/randomi instructions.  DIs use RandomStream.
+	VectorRegister4i          RandCounters;
+
+	FRandomStream             RandStream;
+	int                       UseCount;
+};
+
+struct FVectorVMState
+{
+	uint8 *                     Bytecode;
+	FVectorVMExtFunctionData *  ExtFunctionTable;
+	void **                     UserPtrTable;
+	volatile int32 *            NumOutputPerDataSet;
+	FVectorVMBatchState *       BatchStates;
+	size_t                      NumBytesMalloced;
+	
+	TArrayView<FDataSetMeta>    DataSets;
+
+	uint32                      NumTempRegisters;
+	uint32                      NumConstBuffers;
+	uint32                      NumBytecodeBytes;
+	uint32                      NumOutputDataSets;   //computed in the optimizer
+	uint32                      NumExtFunctions;
+	uint32                      MaxExtFnRegisters;
+	uint32                      NumUserPtrTable;
+	int32                       TotalNumInstances;
+	volatile int32              NumInstancesAssignedToBatches;
+	volatile int32              NumInstancesCompleted;
+
+	size_t                      NumBytesRequiredPerBatch;
+	size_t                      PerBatchRegisterDataBytesRequired;
+	size_t                      PerBatchChunkLocalDataOutputIdxBytesRequired;
+	size_t                      PerBatchChunkLocalNumOutputBytesRequired;
+	size_t                      PerBatchChunkLocalNumExtFnDecodeRegisterBytesRequired;
+
+	int                         NumBatches;
+	int                         MaxInstancesPerChunk;
+	int                         MaxChunksPerBatch;
+
+	VectorVMReallocFn *         ReallocFn;
+	VectorVMFreeFn *            FreeFn;
+
+	struct
+	{
+		int                     LineNum;
+		uint32                  Flags;
+	} Error;
+};
+
+//API FUNCTIONS
+
+//normal functions
+VECTORVM_API FVectorVMState * InitVectorVMState     (FVectorVMInitData *InitData, FVectorVMExternalFnPerInstanceData **OutPerInstanceExtData, FVectorVMSerializeState *SerializeState);
+VECTORVM_API void             FreeVectorVMState     (FVectorVMState *VectorVMState);
+VECTORVM_API void             ExecVectorVMState     (FVectorVMState *VectorVMState, FVectorVMSerializeState *SerializeState, FVectorVMSerializeState *CmpSerializeState);
+VECTORVM_API int              GetNumOutputInstances (FVectorVMState *VectorVMState, int DataSetIdx);
+
+//optimize functions
+VECTORVM_API uint32  OptimizeVectorVMScript      (const uint8 *Bytecode, int BytecodeLen, FVectorVMExtFunctionData *ExtFnIOData, int NumExtFns, FVectorVMOptimizeContext *OptContext, uint32 Flags); //OutContext must be zeroed except the Init struct
+VECTORVM_API void    FreeVectorVMOptimizeContext (FVectorVMOptimizeContext *Context);
+
+//serialize functions
+VECTORVM_API uint32  SerializeVectorVMInputDataSets  (FVectorVMSerializeState *SerializeState, TArrayView<FDataSetMeta> DataSets, FVectorVMConstData *ConstData, int NumConstData); //only use when not calling InitVectorVMState()
+VECTORVM_API uint32  SerializeVectorVMOutputDataSets (FVectorVMSerializeState *SerializeState, TArrayView<FDataSetMeta> DataSets, FVectorVMConstData *ConstData, int NumConstData);
+VECTORVM_API void    SerializeVectorVMWriteToFile    (FVectorVMSerializeState *SerializeState, uint8 WhichStateWritten, const wchar_t *Filename);
+VECTORVM_API void    FreeVectorVMSerializeState      (FVectorVMSerializeState *SerializeState);
+
+#else //NIAGARA_EXP_VM
+
+struct FVectorVMState {
+
+};
+
+VECTORVM_API uint32 SerializeVectorVMInputDataSets  (FVectorVMSerializeState *SerializeState, TArrayView<FDataSetMeta> DataSets, struct FVectorVMConstData *ConstData, int NumConstData); //only use when not calling InitVectorVMState()
+VECTORVM_API uint32 SerializeVectorVMOutputDataSets (FVectorVMSerializeState *SerializeState, TArrayView<FDataSetMeta> DataSets, struct FVectorVMConstData *ConstData, int NumConstData);
+VECTORVM_API void   SerializeVectorVMWriteToFile    (FVectorVMSerializeState *SerializeState, uint8 WhichStateWritten, const wchar_t *Filename);
+VECTORVM_API void   FreeVectorVMSerializeState      (FVectorVMSerializeState *SerializeState);
+
+#endif //NIAGARA_EXP_VM
