Added markup for PVS studio warning. #rb none #preflight none

diff --git a/Engine/Source/Runtime/VectorVM/Public/VectorVM.h b/Engine/Source/Runtime/VectorVM/Public/VectorVM.h
--- a/Engine/Source/Runtime/VectorVM/Public/VectorVM.h
+++ b/Engine/Source/Runtime/VectorVM/Public/VectorVM.h
@@ -315,7 +315,7 @@ private:
 	FDataSetMeta& operator=(const FDataSetMeta&) = delete;
 };
 
-#include "VectorVMExperimental.h"
+#include "VectorVMExperimental.h"  // -V547
 
 //Data the VM will keep on each dataset locally per thread which is then thread safely pushed to it's destination at the end of execution.
 struct FDataSetThreadLocalTempData
