Niagara/VectorVM bug fixes. #rb none #preflight 622146d4a00412627d15dab5

diff --git a/Engine/Plugins/FX/Niagara/Source/Niagara/Private/NiagaraScriptExecutionContext.cpp b/Engine/Plugins/FX/Niagara/Source/Niagara/Private/NiagaraScriptExecutionContext.cpp
--- a/Engine/Plugins/FX/Niagara/Source/Niagara/Private/NiagaraScriptExecutionContext.cpp
+++ b/Engine/Plugins/FX/Niagara/Source/Niagara/Private/NiagaraScriptExecutionContext.cpp
@@ -335,16 +335,24 @@ bool FNiagaraScriptExecutionContextBase::Execute(uint32 NumInstances, const FScr
 #			endif //NIAGARA_EXP_VM
 			
 			//only write under certain circumstances
-			if (OuterObj0 &&
-				OuterObj0->GetName() == L"BoneTest_System" && 
-				OuterObj1 &&
-				OuterObj1->GetName() == L"BoneTest_1" &&
-				ScriptName == L"SpawnScript_0")
-			{
+			//if (OuterObj0 &&
+			//	OuterObj0->GetName() == L"BoneTest_System" && 
+			//	OuterObj1 &&
+			//	OuterObj1->GetName() == L"BoneTest_1" &&
+			//	ScriptName == L"SpawnScript_0")
+			if (0) {
 				SerializeVectorVMOutputDataSets(SerializeState, DataSetMetaTable, VVMConstData, NumVVMConstData);
-				TStringBuilder<256> sb;
-				//sb.Append(Parameters.DebugName);
-				//sb.Append(L"_");
+				TStringBuilder<1024> sb;
+				
+				if (OuterObj0) {
+					sb.Append(OuterObj0->GetName());
+					sb.Append(L"_");
+				}
+				if (OuterObj1) {
+					sb.Append(OuterObj1->GetName());
+					sb.Append(L"_");
+				}
+
 				sb.Append(ScriptName);
 				sb.Append(L"_");
 
@@ -576,28 +584,26 @@ static void PerInsFn(FVectorVMExternalFunctionContext &PerInsFnContext, TArray<F
 	check(SystemInstances);
 	check(*SystemInstances);
 	
-	if (PerInsFnContext.UserPtrTable && UserPtrIdx < PerInsFnContext.NumUserPtrs) {
-		void *SavedUserPtrData       = PerInsFnContext.UserPtrTable[UserPtrIdx];
-		int32 InstanceOffset         = PerInsFnContext.DataSets[0].InstanceOffset; //Apparently the function table is generated based off the first data set, therefore this is safe. (I don't like this - @shawn mcgrath)	
-		int NumInstances             = PerInsFnContext.NumInstances;
-		PerInsFnContext.NumInstances = 1;
-		for (int i = 0; i < NumInstances; ++i) {
-			PerInsFnContext.RegReadCount             = 0;
-			PerInsFnContext.PerInstanceFnInstanceIdx = i;
-
-			int32 InstanceIndex                           = InstanceOffset + PerInsFnContext.StartInstance + i;
-			FNiagaraSystemInstance *Instance              = (**SystemInstances)[InstanceIndex];
-			const FNiagaraPerInstanceDIFuncInfo &FuncInfo = Instance->GetPerInstanceDIFunction(ScriptType, PerInstFunctionIndex);
+	void *SavedUserPtrData       = UserPtrIdx != INDEX_NONE ? PerInsFnContext.UserPtrTable[UserPtrIdx] : NULL;
+	int32 InstanceOffset         = PerInsFnContext.DataSets[0].InstanceOffset; //Apparently the function table is generated based off the first data set, therefore this is safe. (I don't like this - @shawn mcgrath)	
+	int NumInstances             = PerInsFnContext.NumInstances;
+	PerInsFnContext.NumInstances = 1;
+	for (int i = 0; i < NumInstances; ++i) {
+		PerInsFnContext.RegReadCount             = 0;
+		PerInsFnContext.PerInstanceFnInstanceIdx = i;
+
+		int32 InstanceIndex                           = InstanceOffset + PerInsFnContext.StartInstance + i;
+		FNiagaraSystemInstance *Instance              = (**SystemInstances)[InstanceIndex];
+		const FNiagaraPerInstanceDIFuncInfo &FuncInfo = Instance->GetPerInstanceDIFunction(ScriptType, PerInstFunctionIndex);
 		
-			if (UserPtrIdx != INDEX_NONE) {
-				PerInsFnContext.UserPtrTable[UserPtrIdx] = FuncInfo.InstData;
-			}
-			FuncInfo.Function.Execute(PerInsFnContext);
+		if (UserPtrIdx != INDEX_NONE) {
+			PerInsFnContext.UserPtrTable[UserPtrIdx] = FuncInfo.InstData;
 		}
+		FuncInfo.Function.Execute(PerInsFnContext);
+	}
 	
-		if (SavedUserPtrData) {
-			PerInsFnContext.UserPtrTable[UserPtrIdx] = SavedUserPtrData;
-		}
+	if (SavedUserPtrData) {
+		PerInsFnContext.UserPtrTable[UserPtrIdx] = SavedUserPtrData;
 	}
 }
 #else
diff --git a/Engine/Source/Runtime/VectorVM/Private/VectorVMExperimental.inl b/Engine/Source/Runtime/VectorVM/Private/VectorVMExperimental.inl
--- a/Engine/Source/Runtime/VectorVM/Private/VectorVMExperimental.inl
+++ b/Engine/Source/Runtime/VectorVM/Private/VectorVMExperimental.inl
@@ -410,6 +410,7 @@ static uint8 *SetupBatchStatePtrs(FVectorVMState *VVMState, FVectorVMBatchState
 		size_t PtrBeforeExtFnDecodeReg = (size_t)BatchDataPtr;
 		BatchState->ChunkLocalData.ExtFnDecodedReg.RegData       = (FVecReg **)BatchDataPtr;  BatchDataPtr += sizeof(FVecReg *) * VVMState->MaxExtFnRegisters;
 		BatchState->ChunkLocalData.ExtFnDecodedReg.RegInc        = (uint32 *)BatchDataPtr;    BatchDataPtr += sizeof(uint32)    * VVMState->MaxExtFnRegisters;
+		BatchState->ChunkLocalData.ExtFnDecodedReg.DummyRegs     = (FVecReg *)BatchDataPtr;   BatchDataPtr += sizeof(FVecReg)   * VVMState->NumDummyRegsReq;
 		size_t PtrAfterExtFnDecodeReg = (size_t)BatchDataPtr;
 		check(PtrAfterExtFnDecodeReg - PtrBeforeExtFnDecodeReg == VVMState->PerBatchChunkLocalNumExtFnDecodeRegisterBytesRequired);
 	}
@@ -530,7 +531,7 @@ VECTORVM_API FVectorVMState *InitVectorVMState(FVectorVMInitData *InitData, FVec
 	const size_t PerBatchChunkLocalDataOutputIdxBytesRequired	       = sizeof(uint32)                                        * InitData->OptimizeContext->NumOutputDataSets;                  //chunk local bytes required for instance offset
 	const size_t PerBatchChunkLocalNumOutputBytesRequired		       = sizeof(uint32)                                        * InitData->OptimizeContext->NumOutputDataSets;                  //chunk local bytes for num outputs
 	const size_t ConstantBufferSize					                   = sizeof(FVecReg)                                       * InitData->OptimizeContext->NumConstsRemapped;
-	const size_t PerBatchChunkLocalNumExtFnDecodeRegisterBytesRequired = (sizeof(FVecReg *) + sizeof(uint32)) * MaxExtFnRegisters; 
+	const size_t PerBatchChunkLocalNumExtFnDecodeRegisterBytesRequired = (sizeof(FVecReg *) + sizeof(uint32)) * MaxExtFnRegisters + sizeof(FVecReg) * InitData->OptimizeContext->NumDummyRegsReq;
 
 	const size_t BatchOverheadSize                                     = ConstantBufferSize + 
 		                                                                 PerBatchChunkLocalDataOutputIdxBytesRequired + 
@@ -688,6 +689,7 @@ VECTORVM_API FVectorVMState *InitVectorVMState(FVectorVMInitData *InitData, FVec
 	VVMState->NumConstBuffers   = InitData->OptimizeContext->NumConstsRemapped;
 	VVMState->NumOutputDataSets = InitData->OptimizeContext->NumOutputDataSets;
 	VVMState->MaxExtFnRegisters = MaxExtFnRegisters;
+	VVMState->NumDummyRegsReq   = InitData->OptimizeContext->NumDummyRegsReq;
 	VVMState->DataSets          = InitData->DataSets;
 
 	VVMState->UserPtrTable      = InitData->UserPtrTable;
@@ -1342,9 +1344,17 @@ OpCodeSwitch: //I think computed gotos would be a huge win here... maybe write t
 
 						ExtFnCtx.RandCounters             = BatchState->ChunkLocalData.RandCounters;
 						ExtFnCtx.DataSets                 = VVMState->DataSets;
-						ExtFnData->Function->Execute(ExtFnCtx);
 
-						//UE_LOG(LogVectorVM, Warning, TEXT("Num Instances: %d"), NumInstancesThisChunk);
+						//correct invalid registers
+						uint32 DummyRegCount = 0;
+						for (int i = 0; i < ExtFnData->NumInputs + ExtFnData->NumOutputs; ++i) {
+							if (ExtFnCtx.RawVecIndices[i] == 0xFFFF) {
+								ExtFnCtx.RegInc[i] = 0;
+								ExtFnCtx.RegisterData[i] = (uint32 *)(BatchState->ChunkLocalData.ExtFnDecodedReg.DummyRegs + DummyRegCount++);
+							}
+						}
+						check(DummyRegCount <= VVMState->NumDummyRegsReq);
+						ExtFnData->Function->Execute(ExtFnCtx);
 					}
 					InsPtr += 2 + ((ExtFnData->NumInputs + ExtFnData->NumOutputs) << 1);
 				}
diff --git a/Engine/Source/Runtime/VectorVM/Private/VectorVMExperimental_Optimizer.inl b/Engine/Source/Runtime/VectorVM/Private/VectorVMExperimental_Optimizer.inl
--- a/Engine/Source/Runtime/VectorVM/Private/VectorVMExperimental_Optimizer.inl
+++ b/Engine/Source/Runtime/VectorVM/Private/VectorVMExperimental_Optimizer.inl
@@ -884,6 +884,7 @@ VECTORVM_API uint32 OptimizeVectorVMScript(const uint8 *InBytecode, int InByteco
 			break;
 			case EVectorVMOp::external_func_call:
 			{
+				uint32 DummyRegCount = 0;
 				uint8 ExtFnIdx = *OpPtrIn;
 				check(ExtFnIdx < NumExtFns);
 
@@ -898,6 +899,7 @@ VECTORVM_API uint32 OptimizeVectorVMScript(const uint8 *InBytecode, int InByteco
 					if (VecIndices[i] == 0xFFFF)
 					{ //invalid, just write it out
 						OptContext->Intermediate.RegisterUsageBuffer[OptContext->Intermediate.NumRegistersUsed] = 0xFFFF;
+						++DummyRegCount;
 					}
 					else
 					{
@@ -922,9 +924,15 @@ VECTORVM_API uint32 OptimizeVectorVMScript(const uint8 *InBytecode, int InByteco
 				{
 					int Idx = ExtFnIOData[ExtFnIdx].NumInputs + i;
 					check((VecIndices[Idx] & 0x8000) == 0 || VecIndices[Idx] == 0xFFFF); //can't output to a const... 0xFFFF is invalid
+					if (VecIndices[Idx] == 0xFFFF) {
+						++DummyRegCount;
+					}
 					OptContext->Intermediate.RegisterUsageBuffer[OptContext->Intermediate.NumRegistersUsed] = VecIndices[Idx];
 					++OptContext->Intermediate.NumRegistersUsed;
 				}
+				if (DummyRegCount > OptContext->NumDummyRegsReq) {
+					OptContext->NumDummyRegsReq = DummyRegCount;
+				}
 				OptContext->MaxExtFnUsed      = VVM_MAX(OptContext->MaxExtFnUsed, ExtFnIdx);
 				OptContext->MaxExtFnRegisters = VVM_MAX(OptContext->MaxExtFnRegisters, (uint32)(ExtFnIOData[ExtFnIdx].NumInputs + ExtFnIOData[ExtFnIdx].NumOutputs));
 				OpPtrIn += 1 + (ExtFnIOData[ExtFnIdx].NumInputs + ExtFnIOData[ExtFnIdx].NumOutputs) * 2;
@@ -1123,6 +1131,12 @@ VECTORVM_API uint32 OptimizeVectorVMScript(const uint8 *InBytecode, int InByteco
 					{
 						++SSARegCount;
 					}
+					else
+					{
+						//this instruction will be removed later because its output isn't used.  Set the SSA to invalid to avoid messing up
+						//dependency checks before the instruction is removed.
+						OptContext->Intermediate.SSARegisterUsageBuffer[OutputInsRegUse.RegIndices[OutputInsRegUse.NumInputRegisters + j]] = 0xFFFF;
+					}
 				} else {
 					OptContext->Intermediate.SSARegisterUsageBuffer[OutputInsRegUse.RegIndices[OutputInsRegUse.NumInputRegisters + j]] = 0xFFFF;
 				}
@@ -1271,7 +1285,7 @@ VECTORVM_API uint32 OptimizeVectorVMScript(const uint8 *InBytecode, int InByteco
 		}
 	}
 
-	if (0)
+	if (1)
 	{ //Step 7: remove instructions where outputs are never used 
 		int NumRemovedInstructions = 0;
 		FVectorVMOptimizeInsRegUsage RegUsage;
@@ -1410,7 +1424,7 @@ VECTORVM_API uint32 OptimizeVectorVMScript(const uint8 *InBytecode, int InByteco
 		int *InstructionIdxStack = RegToCheckStack + OptContext->Intermediate.NumRegistersUsed;
 		int LowestInstructionIdxForAcquireIdx = OnePastLastInputIdx; //acquire index instructions will be sorted by whichever comes first in the IR... possibly worth checking if re-ordering is more efficient
 
-		{ //Step 9: Find all the acquireindex instructions and re-order them to be executed ASAP
+		if (1) { //Step 9: Find all the acquireindex instructions and re-order them to be executed ASAP
 			int NumAcquireIndexInstructions = 0;
 			for (uint32 i = 0; i < OptContext->Intermediate.NumInstructions; ++i)
 			{
@@ -1490,7 +1504,7 @@ VECTORVM_API uint32 OptimizeVectorVMScript(const uint8 *InBytecode, int InByteco
 			}
 		}
 
-		{ //Step 11: re-order the outputs to be done as early as possible: after the SSA's register's last usage
+		if (1) { //Step 11: re-order the outputs to be done as early as possible: after the SSA's register's last usage
 			for (uint32 OutputInsIdx = 0; OutputInsIdx < OptContext->Intermediate.NumInstructions; ++OutputInsIdx)
 			{
 				FVectorVMOptimizeInstruction *OutputIns = OptContext->Intermediate.Instructions + OutputInsIdx;
@@ -1540,7 +1554,7 @@ VECTORVM_API uint32 OptimizeVectorVMScript(const uint8 *InBytecode, int InByteco
 			}
 		}
 
-		{ //Step 12: re-order all dependent-less instructions to right before their output is used
+		if (1) { //Step 12: re-order all dependent-less instructions to right before their output is used
 			int LastSwapInstructionIdx = -1; //to prevent an infinite loop when one instruction has two or more dependencies and they keep swapping back and forth
 			for (uint32 i = 0; i < OptContext->Intermediate.NumInstructions; ++i)
 			{
@@ -1704,7 +1718,7 @@ VECTORVM_API uint32 OptimizeVectorVMScript(const uint8 *InBytecode, int InByteco
 
 			{ // small list, very stupid bubble sort
 				bool sorted = false;
-				int NumCopyInstructions                        = LastCopyFromInputInsIdx - FirstCopyFromInputInsIdx + 1;
+				int NumCopyInstructions = LastCopyFromInputInsIdx - FirstCopyFromInputInsIdx + 1;
 				while (!sorted) {
 					sorted = true;
 					for (int i = 0; i < NumCopyInstructions - 1; ++i) {
@@ -1856,30 +1870,35 @@ VECTORVM_API uint32 OptimizeVectorVMScript(const uint8 *InBytecode, int InByteco
 			for (int j = 0; j < InsRegUse.NumOutputRegisters; ++j)
 			{
 				uint16 SSARegIdx = OptContext->Intermediate.SSARegisterUsageBuffer[InsRegUse.RegIndices[InsRegUse.NumInputRegisters + j]];
-				uint16 MinimizedRegIdx = 0xFFFF;
-				for (uint16 k = 0; k < NumSSARegistersUsed; ++k)
-				{
-					if (SSAUseMap[k] == 0xFFFF)
+				uint16 OutputRegIdx = InsRegUse.RegIndices[InsRegUse.NumInputRegisters + j];
+				if (SSARegIdx == 0xFFFF) { //"invalid" flag for external functions
+					OptContext->Intermediate.RegisterUsageBuffer[OutputRegIdx] = 0xFFFF;
+				} else {
+					uint16 MinimizedRegIdx = 0xFFFF;
+					for (uint16 k = 0; k < NumSSARegistersUsed; ++k)
 					{
-						SSAUseMap[k] = SSARegIdx;
-						MinimizedRegIdx = k;
-						break;
+						if (SSAUseMap[k] == 0xFFFF)
+						{
+							SSAUseMap[k] = SSARegIdx;
+							MinimizedRegIdx = k;
+							break;
+						}
 					}
-				}
-				check(MinimizedRegIdx != 0xFFFF);
-				uint16 OutputRegIdx = InsRegUse.RegIndices[InsRegUse.NumInputRegisters + j];
-				OptContext->Intermediate.RegisterUsageBuffer[OutputRegIdx] = MinimizedRegIdx;
+					check(MinimizedRegIdx != 0xFFFF);
+				
+					OptContext->Intermediate.RegisterUsageBuffer[OutputRegIdx] = MinimizedRegIdx;
 
-				//change all future instructions to use minimized register index
-				for (uint32 i2 = i + 1; i2 < OptContext->Intermediate.NumInstructions; ++i2)
-				{
-					FVectorVMOptimizeInstruction *Ins2 = OptContext->Intermediate.Instructions + i2;
-					GetRegistersUsedForInstruction(OptContext, Ins2, &InsRegUse2);
-					for (int k = 0; k < InsRegUse2.NumInputRegisters; ++k)
+					//change all future instructions to use minimized register index
+					for (uint32 i2 = i + 1; i2 < OptContext->Intermediate.NumInstructions; ++i2)
 					{
-						if (OptContext->Intermediate.SSARegisterUsageBuffer[InsRegUse2.RegIndices[k]] == OptContext->Intermediate.SSARegisterUsageBuffer[OutputRegIdx])
+						FVectorVMOptimizeInstruction *Ins2 = OptContext->Intermediate.Instructions + i2;
+						GetRegistersUsedForInstruction(OptContext, Ins2, &InsRegUse2);
+						for (int k = 0; k < InsRegUse2.NumInputRegisters; ++k)
 						{
-							OptContext->Intermediate.RegisterUsageBuffer[InsRegUse2.RegIndices[k]] = MinimizedRegIdx;
+							if (OptContext->Intermediate.SSARegisterUsageBuffer[InsRegUse2.RegIndices[k]] == OptContext->Intermediate.SSARegisterUsageBuffer[OutputRegIdx])
+							{
+								OptContext->Intermediate.RegisterUsageBuffer[InsRegUse2.RegIndices[k]] = MinimizedRegIdx;
+							}
 						}
 					}
 				}
