Fix for "Include after first code block" include tool error
#rnx

diff --git a/Engine/Source/Runtime/VectorVM/Public/VectorVM.h b/Engine/Source/Runtime/VectorVM/Public/VectorVM.h
--- a/Engine/Source/Runtime/VectorVM/Public/VectorVM.h
+++ b/Engine/Source/Runtime/VectorVM/Public/VectorVM.h
@@ -10,6 +10,7 @@
 #include "Math/RandomStream.h"
 #include "Misc/ByteSwap.h"
 #include "Templates/AlignmentTemplates.h"
+#include "VectorVMExperimental.h"  
 
 //TODO: move to a per platform header and have VM scale vectorization according to vector width.
 #define VECTOR_WIDTH (128)
@@ -315,7 +316,7 @@ private:
 	FDataSetMeta& operator=(const FDataSetMeta&) = delete;
 };
 
-#include "VectorVMExperimental.h"  // -V547
+
 
 //Data the VM will keep on each dataset locally per thread which is then thread safely pushed to it's destination at the end of execution.
 struct FDataSetThreadLocalTempData
diff --git a/Engine/Source/Runtime/VectorVM/Public/VectorVMExperimental.h b/Engine/Source/Runtime/VectorVM/Public/VectorVMExperimental.h
--- a/Engine/Source/Runtime/VectorVM/Public/VectorVMExperimental.h
+++ b/Engine/Source/Runtime/VectorVM/Public/VectorVMExperimental.h
@@ -5,6 +5,7 @@
 
 typedef void * (VectorVMReallocFn)              (void *Ptr, size_t NumBytes, const char *Filename, int LineNumber);
 typedef void   (VectorVMFreeFn)                 (void *Ptr, const char *Filename, int LineNumber);
+struct FDataSetMeta;
 
 //prototypes for serialization that are required whether or not serialization is enabled
 #if defined(NIAGARA_EXP_VM) || defined(VVM_INCLUDE_SERIALIZATION)
