VectorVM bug fixes & half float support. Experimental VM on by default. #rb Rob.Krajcarski #preflight 6230b204421403c1553829fc

diff --git a/Engine/Plugins/FX/Niagara/Source/Niagara/Niagara.Build.cs b/Engine/Plugins/FX/Niagara/Source/Niagara/Niagara.Build.cs
--- a/Engine/Plugins/FX/Niagara/Source/Niagara/Niagara.Build.cs
+++ b/Engine/Plugins/FX/Niagara/Source/Niagara/Niagara.Build.cs
@@ -67,5 +67,11 @@ public class Niagara : ModuleRules
 				"Slate"
             });
         }
+        PublicDefinitions.AddRange(
+            new string[]
+            {
+                "NIAGARA_EXP_VM",
+            }
+        );
     }
 }
diff --git a/Engine/Plugins/FX/Niagara/Source/Niagara/Private/NiagaraScript.cpp b/Engine/Plugins/FX/Niagara/Source/Niagara/Private/NiagaraScript.cpp
--- a/Engine/Plugins/FX/Niagara/Source/Niagara/Private/NiagaraScript.cpp
+++ b/Engine/Plugins/FX/Niagara/Source/Niagara/Private/NiagaraScript.cpp
@@ -1404,7 +1404,8 @@ FGraphEventRef UNiagaraScript::HandleByteCodeOptimization(bool bShouldForceNow)
 #	ifdef NIAGARA_EXP_VM
 	//this is just necessary because VectorVM doesn't know about FVMExternalFunctionBindingInfo
 	int NumExtFns = CachedScriptVM.CalledVMExternalFunctions.Num();
-	FVectorVMExtFunctionData *ExtFnTable = (FVectorVMExtFunctionData *)alloca(32 + sizeof(FVectorVMExtFunctionData) * NumExtFns); //32 is pointless, it just gets around static analysis triggering for unknown reasons
+	FVectorVMExtFunctionData *ExtFnTable = (FVectorVMExtFunctionData *)FMemory_Alloca(sizeof(FVectorVMExtFunctionData) * NumExtFns); //32 is pointless, it just gets around static analysis triggering for unknown reasons
+	CA_ASSUME(ExtFnTable != nullptr);
 	for (int i = 0; i < NumExtFns; ++i) {
 		ExtFnTable[i].NumInputs = CachedScriptVM.CalledVMExternalFunctions[i].GetNumInputs();
 		ExtFnTable[i].NumOutputs = CachedScriptVM.CalledVMExternalFunctions[i].GetNumOutputs();
diff --git a/Engine/Plugins/FX/Niagara/Source/Niagara/Private/NiagaraScriptExecutionContext.cpp b/Engine/Plugins/FX/Niagara/Source/Niagara/Private/NiagaraScriptExecutionContext.cpp
--- a/Engine/Plugins/FX/Niagara/Source/Niagara/Private/NiagaraScriptExecutionContext.cpp
+++ b/Engine/Plugins/FX/Niagara/Source/Niagara/Private/NiagaraScriptExecutionContext.cpp
@@ -193,8 +193,9 @@ bool FNiagaraScriptExecutionContextBase::Execute(uint32 NumInstances, const FScr
 		FVectorVMSerializeState UESerializeState = { };
 
 #if defined(NIAGARA_EXP_VM) || defined(VVM_INCLUDE_SERIALIZATION)
-		FVectorVMConstData *VVMConstData = (FVectorVMConstData *)alloca(sizeof(FVectorVMConstData) * ConstantBufferTable.BufferSizes.Num());
+		FVectorVMConstData *VVMConstData = (FVectorVMConstData *)FMemory_Alloca(sizeof(FVectorVMConstData) * ConstantBufferTable.BufferSizes.Num());
 		int NumVVMConstData              = ConstantBufferTable.BufferSizes.Num();
+		CA_ASSUME(VVMConstData != nullptr);
 		for (int i = 0; i < NumVVMConstData; ++i)
 		{
 			int32 NumBytes = ConstantBufferTable.BufferSizes[i];
diff --git a/Engine/Source/Runtime/Core/Public/Math/UnrealMathNeon.h b/Engine/Source/Runtime/Core/Public/Math/UnrealMathNeon.h
--- a/Engine/Source/Runtime/Core/Public/Math/UnrealMathNeon.h
+++ b/Engine/Source/Runtime/Core/Public/Math/UnrealMathNeon.h
@@ -1654,6 +1654,8 @@ FORCEINLINE VectorRegister4Float VectorReciprocalSqrtAccurate(const VectorRegist
 	return VectorMultiply(vrsqrtsq_f32(Vec, VectorMultiply(RecipSqrt, RecipSqrt)), RecipSqrt);
 }
 
+#define VectorSqrt(Vec) vsqrtq_f32(Vec)
+
 FORCEINLINE VectorRegister4Double VectorReciprocalSqrtAccurate(const VectorRegister4Double& Vec)
 {
 	// Perform a single pass of Newton-Raphson iteration on the hardware estimate
diff --git a/Engine/Source/Runtime/Core/Public/Math/UnrealMathSSE.h b/Engine/Source/Runtime/Core/Public/Math/UnrealMathSSE.h
--- a/Engine/Source/Runtime/Core/Public/Math/UnrealMathSSE.h
+++ b/Engine/Source/Runtime/Core/Public/Math/UnrealMathSSE.h
@@ -1936,6 +1936,8 @@ FORCEINLINE VectorRegister4Double VectorReciprocalSqrt(const VectorRegister4Doub
 	return VectorDivide(GlobalVectorConstants::DoubleOne, Result);
 }
 
+#define VectorSqrt(Vec) _mm_sqrt_ps(Vec)
+
 /**
  * Computes an estimate of the reciprocal of a vector (component-wise) and returns the result.
  *
diff --git a/Engine/Source/Runtime/Experimental/Chaos/Public/Chaos/VectorUtility.h b/Engine/Source/Runtime/Experimental/Chaos/Public/Chaos/VectorUtility.h
--- a/Engine/Source/Runtime/Experimental/Chaos/Public/Chaos/VectorUtility.h
+++ b/Engine/Source/Runtime/Experimental/Chaos/Public/Chaos/VectorUtility.h
@@ -72,10 +72,10 @@ FORCEINLINE VectorRegister4Float VectorMoveLh(const VectorRegister4Float& A, con
  * @param A	1st vector
   * @return		VectorRegister4Float( sqrt(A.x), sqrt(A.y), sqrt(A.z), sqrt(A.w))
   */
-FORCEINLINE VectorRegister4Float VectorSqrt(const VectorRegister4Float& A)
-{
-	return vsqrtq_f32(A);
-}
+//FORCEINLINE VectorRegister4Float VectorSqrt(const VectorRegister4Float& A)
+//{
+//	return vsqrtq_f32(A);
+//}
 
 
 #else // PLATFORM_ENABLE_VECTORINTRINSICS_NEON
@@ -110,10 +110,10 @@ FORCEINLINE VectorRegister4Float VectorMoveLh(const VectorRegister4Float& A, con
  * @param A	1st vector
   * @return		VectorRegister4Float( sqrt(A.x), sqrt(A.y), sqrt(A.z), sqrt(A.w))
   */
-FORCEINLINE VectorRegister4Float VectorSqrt(const VectorRegister4Float& A)
-{
-	return _mm_sqrt_ps(A);
-}
+//FORCEINLINE VectorRegister4Float VectorSqrt(const VectorRegister4Float& A)
+//{
+//	return _mm_sqrt_ps(A);
+//}
 
 #endif // PLATFORM_ENABLE_VECTORINTRINSICS_NEON
 
diff --git a/Engine/Source/Runtime/VectorVM/Private/VectorVM.cpp b/Engine/Source/Runtime/VectorVM/Private/VectorVM.cpp
--- a/Engine/Source/Runtime/VectorVM/Private/VectorVM.cpp
+++ b/Engine/Source/Runtime/VectorVM/Private/VectorVM.cpp
@@ -10,6 +10,7 @@
 #include "HAL/PlatformFileManager.h"
 #include "GenericPlatform/GenericPlatformFile.h"
 #include "Async/ParallelFor.h"
+#include "Math/UnrealPlatformMathSSE.h"
 
 IMPLEMENT_MODULE(FDefaultModuleImpl, VectorVM);
 
diff --git a/Engine/Source/Runtime/VectorVM/Private/VectorVMExperimental.inl b/Engine/Source/Runtime/VectorVM/Private/VectorVMExperimental.inl
--- a/Engine/Source/Runtime/VectorVM/Private/VectorVMExperimental.inl
+++ b/Engine/Source/Runtime/VectorVM/Private/VectorVMExperimental.inl
@@ -325,15 +325,16 @@ VVMSet_m128iConst(  AlmostTwoBits   , 0x3fffffff);
 #define VVM_vecBoolToFloat(v)        VectorSelect(v, VectorSet1(1.f), VectorZeroFloat());
 #define VVM_vecIntToBool(v)          VectorIntCompareGT(v, VectorSetZero())
 #define VVM_vecBoolToInt(v)          VectorIntSelect(v, VectorIntSet1(1), VectorSetZero())
-#define VVM_vecSqrtFast(v)           VectorReciprocal(VectorReciprocalSqrt(v))
-#define VVM_vecACosFast(v)           VectorATan2(VVM_vecSqrtFast(VectorMultiply(VectorSubtract(VVM_m128Const(One), v), VectorAdd(VVM_m128Const(One), v))), v)
+#define VVM_vecACosFast(v)           VectorATan2(VectorSqrt(VectorMultiply(VectorSubtract(VVM_m128Const(One), v), VectorAdd(VVM_m128Const(One), v))), v)
 //safe instructions -- handle divide by zero "gracefully" by returning 0
-#define VVM_safeIns_div(v0, v1)      VectorSelect(VectorCompareGT(VectorAbs(v1), VVM_m128Const(Epsilon)), VectorDivide(v0, v1)    , VectorZeroFloat())
-#define VVM_safeIns_rcp(v)           VectorSelect(VectorCompareGT(VectorAbs(v) , VVM_m128Const(Epsilon)), VectorReciprocal(v)     , VectorZeroFloat())
-#define VVM_safe_sqrt(v)             VectorSelect(VectorCompareGT(VectorAbs(v) , VVM_m128Const(Epsilon)), VVM_vecSqrtFast(v)      , VectorZeroFloat())
-#define VVM_safe_log(v)              VectorSelect(VectorCompareGT(VectorAbs(v) , VectorZeroFloat())     , VectorLog(v)		       , VectorZeroFloat())
-#define VVM_safe_pow(v0, v1)         VectorSelect(VectorCompareGT(VectorAbs(v1), VVM_m128Const(Epsilon)), VectorPow(v0, v1)       , VectorZeroFloat())
-#define VVM_safe_rsq(v)              VectorSelect(VectorCompareGT(VectorAbs(v) , VVM_m128Const(Epsilon)), VectorReciprocalSqrt(v) , VectorZeroFloat())
+#define VVM_safeIns_div(v0, v1)      VectorSelect(VectorCompareGT(VectorAbs(v1), VVM_m128Const(Epsilon)), VectorDivide(v0, v1)                  , VectorZeroFloat())
+#define VVM_safeIns_rcp(v)           VectorSelect(VectorCompareGT(VectorAbs(v) , VVM_m128Const(Epsilon)), VectorDivide(VVM_m128Const(One), v)   , VectorZeroFloat())
+#define VVM_safe_sqrt(v)             VectorSelect(VectorCompareGT(VectorAbs(v) , VVM_m128Const(Epsilon)), VectorSqrt(v)                         , VectorZeroFloat())
+#define VVM_safe_log(v)              VectorSelect(VectorCompareGT(VectorAbs(v) , VectorZeroFloat())     , VectorLog(v)		                    , VectorZeroFloat())
+#define VVM_safe_pow(v0, v1)         VectorSelect(VectorCompareGT(VectorAbs(v1), VVM_m128Const(Epsilon)), VectorPow(v0, v1)                     , VectorZeroFloat())
+#define VVM_safe_rsq(v)              VectorSelect(VectorCompareGT(VectorAbs(v) , VVM_m128Const(Epsilon)), VectorReciprocalSqrt(v)               , VectorZeroFloat())
+
+#define VVMDebugBreakIf(expr)	if ((expr)) { PLATFORM_BREAK(); }
 
 static void VVMMemCpy(void *dst, void *src, size_t bytes)
 {
@@ -442,6 +443,9 @@ static bool AssignInstancesToBatch(FVectorVMState *VVMState, FVectorVMBatchState
 		int NumAssignedInstances    = FPlatformAtomics::InterlockedCompareExchange(&VVMState->NumInstancesAssignedToBatches, OldNumAssignedInstances + MaxInstancesPerBatch, OldNumAssignedInstances);
 		if (NumAssignedInstances == OldNumAssignedInstances)
 		{
+			//we allow "overflow" on the thread-local BatchState to keep access to the contended variable: VVMState->NumInstancesAssignedToBatches 
+			//down to a single atomic.  Letting BatchState->NumInstances go < 0 is fine because this will return false and Exec() will not run
+			//for this BatchState.
 			BatchState->StartInstance = OldNumAssignedInstances;
 			BatchState->NumInstances  = MaxInstancesPerBatch;
 			if (BatchState->StartInstance + BatchState->NumInstances > VVMState->TotalNumInstances) {
@@ -455,6 +459,7 @@ static bool AssignInstancesToBatch(FVectorVMState *VVMState, FVectorVMBatchState
 			return true;
 		}
 	} while (SanityCount++ < (1 << 30));
+	VVMDebugBreakIf(SanityCount > (1 << 30) - 1);
 	return false;
 }
 
@@ -477,19 +482,19 @@ static void SetupRandStateForBatch(FVectorVMBatchState *BatchState)
 	uint64 pcg_state = FPlatformTime::Cycles64();
 	uint64 pcg_inc   = (((uint64)BatchState << 32) ^ 0XCAFEF00DD15EA5E5U) | 1;
 	pcg_state ^= (FPlatformTime::Cycles64() << 32ULL);
-	//use psuedo-pcg to setup a state for xorwow... lol!
-	for (int i = 0; i < 5; ++i)
+	//use psuedo-pcg to setup a state for xorwow
+	for (int i = 0; i < 5; ++i) //loop for xorwow internal state
 	{
-		uint32 values[4];
+		MS_ALIGN(16) uint32 Values[4] GCC_ALIGN(16);
 		for (int j = 0; j < 4; ++j)
 		{
 			uint64 old_state   = pcg_state;
 			pcg_state          = old_state * 6364136223846793005ULL + pcg_inc;
 			uint32 xor_shifted = (uint32)(((old_state >> 18U) ^ old_state) >> 27U);
 			uint32 rot         = old_state >> 59U;
-			values[j]          = (xor_shifted >> rot) | (xor_shifted << ((0U - rot) & 31));
+			Values[j]          = (xor_shifted >> rot) | (xor_shifted << ((0U - rot) & 31));
 		}
-		VectorIntStore(MakeVectorRegisterInt(values[0], values[1], values[2], values[3]), BatchState->RandState + i);
+		VectorIntStore(*(VectorRegister4i *)Values, BatchState->RandState + i);
 	}
 	BatchState->RandCounters = MakeVectorRegisterInt64(pcg_inc, pcg_state);
 	BatchState->RandStream.GenerateNewSeed();
@@ -753,35 +758,7 @@ VECTORVM_API FVectorVMState *InitVectorVMState(FVectorVMInitData *InitData, FVec
 }
 
 
-static void ExecVVMBatch(FVectorVMState *VVMState, int ExecIdx, FVectorVMSerializeState *SerializeState, FVectorVMSerializeState *CmpSerializeState)
-{
-#	if defined(VVM_INCLUDE_SERIALIZATION) && !defined(VVM_SERIALIZE_NO_WRITE)
-#		ifdef VVM_SERIALIZE_PERF
-#			define serializeIns(Type, NumParams)	
-#			define serializeRegUsed(RegIdx, Type)
-#		else //VVM_SERIALIZE_PERF
-#			define serializeIns(Type, NumParams)	if (SerializeState)                                                                         \
-													{                                                                                           \
-														for (int vi = 0; vi <= (int)(NumParams); ++vi)                                          \
-														{                                                                                       \
-															if ((VecIndices[vi] & 0x8000) == 0)											        \
-															{                                                                                   \
-																SerializeState->TempRegFlags[VecIndices[vi]] = VVMRegFlag_Clean + (Type);       \
-															}                                                                                   \
-														}                                                                                       \
-													}
-
-#			define serializeRegUsed(RegIdx, Type)	if (SerializeState)                                                                         \
-													{                                                                                           \
-														SerializeState->TempRegFlags[RegIdx] = VVMRegFlag_Clean + (Type);                       \
-													}
-#		endif //VVM_SERIALIZE_PERF
-#	else //VVM_INCLUDE_SERIALIZATION
-#		define serializeIns(Type, NumParams)		
-#		define serializeRegUsed(RegIdx, Type)
-#	endif //VVM_INCLUDE_SERIALIZATION
-
-#	define execVecIns1f(ins_)						serializeIns(0, 1)                                                  \
+#	define execVecIns1f(ins_)						VVMSer_instruction(0, 1)                                            \
 													for (int i = 0; i < NumLoops; ++i)                                  \
 													{                                                                   \
 														VectorRegister4f r0  = VectorLoad(&VecReg[0][i & RegInc[0]].v); \
@@ -790,7 +767,7 @@ static void ExecVVMBatch(FVectorVMState *VVMState, int ExecIdx, FVectorVMSeriali
 													}                                                                   \
 													InsPtr += 4;
 
-#	define execVecIns2f(ins_)						serializeIns(0, 2)                                                  \
+#	define execVecIns2f(ins_)						VVMSer_instruction(0, 2)                                            \
 													for (int i = 0; i < NumLoops; ++i)                                  \
 													{                                                                   \
 														VectorRegister4f r0  = VectorLoad(&VecReg[0][i & RegInc[0]].v); \
@@ -800,7 +777,7 @@ static void ExecVVMBatch(FVectorVMState *VVMState, int ExecIdx, FVectorVMSeriali
 													}                                                                   \
 													InsPtr += 6;
 
-#	define execVecIns3f(ins_)						serializeIns(0, 3)                                                  \
+#	define execVecIns3f(ins_)						VVMSer_instruction(0, 3)                                            \
 													for (int i = 0; i < NumLoops; ++i)                                  \
 													{                                                                   \
 														VectorRegister4f r0  = VectorLoad(&VecReg[0][i & RegInc[0]].v); \
@@ -811,7 +788,7 @@ static void ExecVVMBatch(FVectorVMState *VVMState, int ExecIdx, FVectorVMSeriali
 													}                                                                   \
 													InsPtr += 8;
 
-#	define execVecIns1i(ins_)						serializeIns(1, 1)                                                     \
+#	define execVecIns1i(ins_)						VVMSer_instruction(1, 1)                                               \
 													for (int i = 0; i < NumLoops; ++i)                                     \
 													{                                                                      \
 														VectorRegister4i r0  = VectorIntLoad(&VecReg[0][i & RegInc[0]].v); \
@@ -820,7 +797,7 @@ static void ExecVVMBatch(FVectorVMState *VVMState, int ExecIdx, FVectorVMSeriali
 													}                                                                      \
 													InsPtr += 4;
 
-#	define execVecIns2i(ins_)						serializeIns(1, 2)                                                     \
+#	define execVecIns2i(ins_)						VVMSer_instruction(1, 2)                                               \
 													for (int i = 0; i < NumLoops; ++i)                                     \
 													{                                                                      \
 														VectorRegister4i r0  = VectorIntLoad(&VecReg[0][i & RegInc[0]].v); \
@@ -830,7 +807,7 @@ static void ExecVVMBatch(FVectorVMState *VVMState, int ExecIdx, FVectorVMSeriali
 													}                                                                      \
 													InsPtr += 6;
 
-#	define execVecIns3i(ins_)						serializeIns(1, 3)                                                     \
+#	define execVecIns3i(ins_)						VVMSer_instruction(1, 3)                                               \
 													for (int i = 0; i < NumLoops; ++i)                                     \
 													{                                                                      \
 														VectorRegister4i r0  = VectorIntLoad(&VecReg[0][i & RegInc[0]].v); \
@@ -841,7 +818,26 @@ static void ExecVVMBatch(FVectorVMState *VVMState, int ExecIdx, FVectorVMSeriali
 													}                                                                      \
 													InsPtr += 8;
 
+#		define VVMDecodeInstructionRegisters(Bytecode, RegData, RegIncMask) do {                                                                                                                                                            \
+			uint32 VecOffsets[4];                                                                                                                                                                                                           \
+			VectorRegister4i VecIndicesIn4  = VectorIntLoad(Bytecode);                                                                          /* 16 bit inputs.  15 bits for index, 1 high bit for const/reg flag (0: reg, 1: const) */   \
+			VectorRegister4i VecIndices4    = VectorIntExpandLow16To32(VecIndicesIn4);                                                          /* 4-wide 32 bit version of the inputs, bits 16:31 are 0 */                                 \
+			VectorRegister4i RegIncInv4     = VectorIntSubtract(VectorSetZero(), VectorShiftRightImmArithmetic(VecIndices4, 15));               /* Sets the inverse of what we need: 0xFF... is const, 0 is reg */                          \
+			VectorRegister4i RegInc4        = VectorIntXor(RegIncInv4, VVM_m128iConst(FMask));                                                  /* Whether to increment the index counter: 0xFF... for registers, and 0 for const */        \
+			VectorRegister4i VecRegIndices4 = VectorIntAnd(VecIndices4, VVM_m128iConst(RegOffsetMask));                                         /* Only the register index, the const/reg flag is stripped off */                           \
+			VectorRegister4i ConstOffset4   = VectorIntAnd(VecRegIndices4, RegIncInv4);                                                         /* Only the const offsets, all registers are masked out */                                  \
+			VectorRegister4i TempRegOffset4 = VectorIntAnd(VectorIntAdd(NumConsts4, VectorIntMultiply(VecRegIndices4, NumLoops4)), RegInc4);    /* Only the register offsets, all consts are masked out */                                  \
+			VectorRegister4i OptRegOffset4  = VectorIntOr(ConstOffset4, TempRegOffset4);                                                        /* Blended (sse4 would be nice) const and temp register offsets */                          \
+			VectorIntStore(RegInc4, RegIncMask);                                                                                                                                                                                            \
+			VectorIntStore(OptRegOffset4, VecOffsets);                                                                                                                                                                                      \
+			(RegData)[0] = BatchState->RegisterData + VecOffsets[0];                                                                                                                                                                        \
+			(RegData)[1] = BatchState->RegisterData + VecOffsets[1];                                                                                                                                                                        \
+			(RegData)[2] = BatchState->RegisterData + VecOffsets[2];                                                                                                                                                                        \
+			(RegData)[3] = BatchState->RegisterData + VecOffsets[3];                                                                                                                                                                        \
+		} while (0)
 
+static void ExecVVMBatch(FVectorVMState *VVMState, int ExecIdx, FVectorVMSerializeState *SerializeState, FVectorVMSerializeState *CmpSerializeState)
+{
 	FVectorVMBatchState *BatchState = nullptr;
 	int BatchIdx = -1;
 	//check to see if we can reuse a batch's memory that's finished executing
@@ -893,7 +889,6 @@ static void ExecVVMBatch(FVectorVMState *VVMState, int ExecIdx, FVectorVMSeriali
 	int NumChunksThisBatch     = (BatchState->NumInstances + VVMState->MaxInstancesPerChunk - 1) / VVMState->MaxInstancesPerChunk;
 	
 	uint32 RegInc[4];
-		
 	VectorRegister4i NumConsts4 = VectorIntSet1(VVMState->NumConstBuffers);
 	
 	for (int ChunkIdxThisBatch = 0; ChunkIdxThisBatch < NumChunksThisBatch; ++ChunkIdxThisBatch, StartInstanceThisChunk += VVMState->MaxInstancesPerChunk)
@@ -912,24 +907,6 @@ static void ExecVVMBatch(FVectorVMState *VVMState, int ExecIdx, FVectorVMSeriali
 			BatchState->ChunkLocalData.NumOutputPerDataSet[i] = 0;
 		}
 		
-#		define VVMDecodeInstructionRegisters(Bytecode, RegData, RegIncMask) do {                                                                                                                                                            \
-			uint32 VecOffsets[4];                                                                                                                                                                                                           \
-			VectorRegister4i VecIndicesIn4  = VectorIntLoad(Bytecode);                                                                          /* 16 bit inputs.  15 bits for index, 1 high bit for const/reg flag (0: reg, 1: const) */   \
-			VectorRegister4i VecIndices4    = VectorIntExpandLow16To32(VecIndicesIn4);                                                          /* 4-wide 32 bit version of the inputs, bits 16:31 are 0 */                                 \
-			VectorRegister4i RegIncInv4     = VectorIntSubtract(VectorSetZero(), VectorShiftRightImmArithmetic(VecIndices4, 15));               /* Sets the inverse of what we need: 0xFF... is const, 0 is reg */                          \
-			VectorRegister4i RegInc4        = VectorIntXor(RegIncInv4, VVM_m128iConst(FMask));                                                  /* Whether to increment the index counter: 0xFF... for registers, and 0 for const */        \
-			VectorRegister4i VecRegIndices4 = VectorIntAnd(VecIndices4, VVM_m128iConst(RegOffsetMask));                                         /* Only the register index, the const/reg flag is stripped off */                           \
-			VectorRegister4i ConstOffset4   = VectorIntAnd(VecRegIndices4, RegIncInv4);                                                         /* Only the const offsets, all registers are masked out */                                  \
-			VectorRegister4i TempRegOffset4 = VectorIntAnd(VectorIntAdd(NumConsts4, VectorIntMultiply(VecRegIndices4, NumLoops4)), RegInc4);    /* Only the register offsets, all consts are masked out */                                  \
-			VectorRegister4i OptRegOffset4  = VectorIntOr(ConstOffset4, TempRegOffset4);                                                        /* Blended (sse4 would be nice) const and temp register offsets */                          \
-			VectorIntStore(RegInc4, RegIncMask);                                                                                                                                                                                            \
-			VectorIntStore(OptRegOffset4, VecOffsets);                                                                                                                                                                                      \
-			(RegData)[0] = BatchState->RegisterData + VecOffsets[0];                                                                                                                                                                        \
-			(RegData)[1] = BatchState->RegisterData + VecOffsets[1];                                                                                                                                                                        \
-			(RegData)[2] = BatchState->RegisterData + VecOffsets[2];                                                                                                                                                                        \
-			(RegData)[3] = BatchState->RegisterData + VecOffsets[3];                                                                                                                                                                        \
-		} while (0)
-
 		FVecReg *VecReg[4];
 		
 		while (InsPtr < InsPtrEnd)
@@ -979,7 +956,6 @@ OpCodeSwitch: //I think computed gotos would be a huge win here... maybe write t
 				case EVectorVMOp::step:                             execVecIns2f(VVM_vecStep);          break;
 				case EVectorVMOp::random:
 #					if defined(VVM_INCLUDE_SERIALIZATION) && !defined(VVM_SERIALIZE_NO_WRITE)
-#					ifndef VVM_SERIALIZE_PERF
 					if (SerializeState && (SerializeState->Flags & VVMSer_SyncRandom) && CmpSerializeState && CmpSerializeState->NumInstructions >= SerializeState->NumInstructions && CmpSerializeState->NumTempRegisters > VecIndices[1])
 					{
 						FVectorVMSerializeInstruction *CmpIns = nullptr;
@@ -1019,7 +995,6 @@ OpCodeSwitch: //I think computed gotos would be a huge win here... maybe write t
 							goto VVM_EndSyncRandom;
 						}
 					}
-#					endif //VVM_SERIALIZE_PERF
 #					endif //VVM_INCLUDE_SERIALIZATION
 					for (int i = 0; i < NumLoops; ++i)
 					{
@@ -1031,7 +1006,7 @@ OpCodeSwitch: //I think computed gotos would be a huge win here... maybe write t
 #					if defined(VVM_INCLUDE_SERIALIZATION) && !defined(VVM_SERIALIZE_NO_WRITE)
 					VVM_EndSyncRandom:
 #					endif //VVM_INCLUDE_SERIALIZATION
-					serializeRegUsed(VecIndices[1], 0);
+					VVMSer_regUsed(VecIndices[1], 0);
 					InsPtr += 4;
 				break;
 				case EVectorVMOp::noise:                            check(false);                       break;
@@ -1046,7 +1021,7 @@ OpCodeSwitch: //I think computed gotos would be a huge win here... maybe write t
 				case EVectorVMOp::subi:                             execVecIns2i(VectorIntSubtract);    break;
 				case EVectorVMOp::muli:                             execVecIns2i(VectorIntMultiply);    break;
 				case EVectorVMOp::divi: {
-					serializeIns(1, 2)
+					VVMSer_instruction(1, 2)
 					//@TODO: convert to double and div 4 wide
 					for (int i = 0; i < NumLoops; ++i)
 					{
@@ -1073,7 +1048,7 @@ OpCodeSwitch: //I think computed gotos would be a huge win here... maybe write t
 				case EVectorVMOp::signi:                            execVecIns1i(VectorIntSign);        break;
 				case EVectorVMOp::randomi: {
 					//@TODO: serialize syncing, no test cases yet
-					serializeIns(1, 1)
+					VVMSer_instruction(1, 1)
 					for (int i = 0; i < NumLoops; ++i)
 					{
 						VecReg[0][i].i = VVMXorwowStep(BatchState);
@@ -1101,7 +1076,7 @@ OpCodeSwitch: //I think computed gotos would be a huge win here... maybe write t
 						((int *)VecReg[2])[(i << 2) + 2] = ((int *)VecReg[0])[idx0 + 2] << ((int *)VecReg[1])[idx1 + 2];
 						((int *)VecReg[2])[(i << 2) + 3] = ((int *)VecReg[0])[idx0 + 3] << ((int *)VecReg[1])[idx1 + 3];
 					}
-					serializeRegUsed(VecIndices[2], 1);
+					VVMSer_regUsed(VecIndices[2], 1);
 					InsPtr += 6;
 					break;
 				case EVectorVMOp::bit_rshift:
@@ -1114,7 +1089,7 @@ OpCodeSwitch: //I think computed gotos would be a huge win here... maybe write t
 						((int *)VecReg[2])[(i << 2) + 2] = ((int *)VecReg[0])[idx0 + 2] >> ((int *)VecReg[1])[idx1 + 2];
 						((int *)VecReg[2])[(i << 2) + 3] = ((int *)VecReg[0])[idx0 + 3] >> ((int *)VecReg[1])[idx1 + 3];
 					}
-					serializeRegUsed(VecIndices[2], 1);
+					VVMSer_regUsed(VecIndices[2], 1);
 					InsPtr += 6;
 					break;
 				case EVectorVMOp::logic_and:                        execVecIns2i(VectorIntAnd);         break;
@@ -1126,7 +1101,7 @@ OpCodeSwitch: //I think computed gotos would be a huge win here... maybe write t
 					{
 						VecReg[1][i].i = VectorFloatToInt(VecReg[0][i & RegInc[0]].v);
 					}
-					serializeRegUsed(VecIndices[1], 1);
+					VVMSer_regUsed(VecIndices[1], 1);
 					InsPtr += 4;
 				break;
 				case EVectorVMOp::i2f:
@@ -1134,7 +1109,7 @@ OpCodeSwitch: //I think computed gotos would be a huge win here... maybe write t
 					{
 						VecReg[1][i].v = VectorIntToFloat(VecReg[0][i & RegInc[0]].i);
 					}
-					serializeRegUsed(VecIndices[1], 0);
+					VVMSer_regUsed(VecIndices[1], 0);
 					InsPtr += 4;
 				break;
 				case EVectorVMOp::f2b:                              execVecIns1f(VVM_vecFloatToBool);   break;
@@ -1143,7 +1118,7 @@ OpCodeSwitch: //I think computed gotos would be a huge win here... maybe write t
 				case EVectorVMOp::b2i:                              execVecIns1i(VVM_vecBoolToInt);     break;
 				case EVectorVMOp::inputdata_float:
 				case EVectorVMOp::inputdata_int32: {
-					uint8 RegType             = (uint8)OpCode - (uint8)EVectorVMOp::inputdata_float;
+					uint8 RegType             = (uint8)OpCode - (uint8)EVectorVMOp::inputdata_float;  //0: float, 1: int, 2: half
 					uint16 DataSetIdx         = *(uint16 *)(InsPtr    );
 					uint16 InputRegIdx        = *(uint16 *)(InsPtr + 2);
 					uint16 DestRegIdx         = *(uint16 *)(InsPtr + 4);
@@ -1152,13 +1127,27 @@ OpCodeSwitch: //I think computed gotos would be a huge win here... maybe write t
 					size_t DstIdx             = NumLoops * DestRegIdx;
 					uint32 **InputBuffers     = (uint32 **)VVMState->DataSets[DataSetIdx].InputRegisters.GetData();
 					VVMMemCpy(BatchState->RegisterData + VVMState->NumConstBuffers + DstIdx, InputBuffers[InputRegIdx + InputRegTypeOffset] + StartInstanceThisChunk + InstanceOffset, sizeof(FVecReg) * NumLoops);
-					serializeRegUsed(DestRegIdx, RegType);
+					VVMSer_regUsed(DestRegIdx, RegType);
+					InsPtr += 6;
+				} break;
+				case EVectorVMOp::inputdata_half: {
+					uint8 RegType             = 2;
+					uint16 DataSetIdx         = *(uint16 *)(InsPtr    );
+					uint16 InputRegIdx        = *(uint16 *)(InsPtr + 2);
+					uint16 DestRegIdx         = *(uint16 *)(InsPtr + 4);
+					uint32 InstanceOffset     = VVMState->DataSets[DataSetIdx].InstanceOffset;
+					uint32 InputRegTypeOffset = VVMState->DataSets[DataSetIdx].InputRegisterTypeOffsets[RegType];
+					size_t DstIdx             = NumLoops * DestRegIdx;
+					uint16 **InputBuffers     = (uint16 **)VVMState->DataSets[DataSetIdx].InputRegisters.GetData();
+					for (int i = 0; i < NumLoops; ++i) {
+						FPlatformMath::VectorLoadHalf((float *)(BatchState->RegisterData + VVMState->NumConstBuffers + DstIdx + i), (InputBuffers[InputRegIdx + InputRegTypeOffset] + StartInstanceThisChunk + InstanceOffset + i));
+					}
+					VVMSer_regUsed(DestRegIdx, 0);
 					InsPtr += 6;
 				} break;
-				case EVectorVMOp::inputdata_half:                   check(false);                       break;
 				case EVectorVMOp::inputdata_noadvance_float:
 				case EVectorVMOp::inputdata_noadvance_int32: {
-					uint8 RegType                 = (uint8)OpCode - (uint8)EVectorVMOp::inputdata_noadvance_float;
+					uint8 RegType                 = (uint8)OpCode - (uint8)EVectorVMOp::inputdata_noadvance_float; //0: float, 1: int, 2: half
 					uint16 DataSetIdx             = *(uint16 *)(InsPtr    );
 					uint16 InputRegIdx            = *(uint16 *)(InsPtr + 2);
 					uint16 DestRegIdx             = *(uint16 *)(InsPtr + 4);
@@ -1172,7 +1161,7 @@ OpCodeSwitch: //I think computed gotos would be a huge win here... maybe write t
 					{
 						VectorIntStoreAligned(input_val4, BatchState->RegisterData + VVMState->NumConstBuffers + DstIdx + i);
 					}
-					serializeRegUsed(DestRegIdx, RegType);
+					VVMSer_regUsed(DestRegIdx, RegType);
 					InsPtr += 6;
 				} break;
 				case EVectorVMOp::inputdata_noadvance_half:			check(false);				        break;
@@ -1208,7 +1197,23 @@ OpCodeSwitch: //I think computed gotos would be a huge win here... maybe write t
 					}
 					InsPtr += 8;
 				} break;
-				case EVectorVMOp::outputdata_half:					check(false);				        break;
+				case EVectorVMOp::outputdata_half: {
+					uint8 RegType             = 2;
+					uint16 DataSetIdx         = VecIndices[0];
+					int *DstIdxReg            = (int *)VecReg[1];
+					uint32 NumOutputInstances = BatchState->ChunkLocalData.NumOutputPerDataSet[DataSetIdx];
+					uint32 InstanceOffset     = BatchState->ChunkLocalData.StartingOutputIdxPerDataSet[DataSetIdx];
+					uint32 RegTypeOffset      = VVMState->DataSets[DataSetIdx].OutputRegisterTypeOffsets[RegType];
+					uint16 **OutputBuffers    = (uint16 **)VVMState->DataSets[DataSetIdx].OutputRegisters.GetData();
+					uint16 *DstReg            = OutputBuffers[RegTypeOffset + VecIndices[3]] + InstanceOffset;
+					uint32 *SrcReg            = (uint32 *)VecReg[2];
+
+					for (uint32 i = 0; i < NumOutputInstances; ++i) //scalar, not ideal
+					{
+						FPlatformMath::StoreHalf(DstReg + i, ((float *)SrcReg)[DstIdxReg[i & RegInc[1]] & RegInc[2]]);
+					}
+					InsPtr += 8;
+				} break;
 				case EVectorVMOp::acquireindex:
 				{
 					uint32 NumOutputInstances = 0;
@@ -1232,7 +1237,7 @@ OpCodeSwitch: //I think computed gotos would be a huge win here... maybe write t
 					}
 					BatchState->ChunkLocalData.StartingOutputIdxPerDataSet[DataSetIdx] = VVMState->DataSets[DataSetIdx].InstanceOffset + FPlatformAtomics::InterlockedAdd(VVMState->NumOutputPerDataSet + DataSetIdx, NumOutputInstances);
 					BatchState->ChunkLocalData.NumOutputPerDataSet[DataSetIdx] += NumOutputInstances;
-					serializeRegUsed(VecIndices[2], VVMRegFlag_Int | VVMRegFlag_Index);
+					VVMSer_regUsed(VecIndices[2], VVMRegFlag_Int | VVMRegFlag_Index);
 					InsPtr += 8;
 				}
 				break;
@@ -1301,7 +1306,7 @@ OpCodeSwitch: //I think computed gotos would be a huge win here... maybe write t
 						for (int i = 0; i < ExtFnData->NumOutputs; ++i)
 						{
 							uint16 RegIdx = ((uint16 *)(InsPtr + 2))[ExtFnData->NumInputs + i];
-							serializeRegUsed(RegIdx, 0); //assume float
+							VVMSer_regUsed(RegIdx, 0); //assume float
 						}
 					} else {
 #					else //VVM_INCLUDE_SERIALIZATION
@@ -1362,7 +1367,7 @@ OpCodeSwitch: //I think computed gotos would be a huge win here... maybe write t
 				case EVectorVMOp::exec_index:
 				{
 					int RegIdx = NumLoops * *((uint16 *)InsPtr);
-					serializeIns(1, 0);
+					VVMSer_instruction(1, 0);
 					VectorRegister4i StartInstance4 = VectorIntAdd(VectorIntSet1(StartInstanceThisChunk), VVM_m128iConst(ZeroOneTwoThree));
 					for (int i = 0; i < NumLoops; ++i)
 					{
@@ -1379,7 +1384,7 @@ OpCodeSwitch: //I think computed gotos would be a huge win here... maybe write t
 				case EVectorVMOp::exit_stat_scope:                                                  break;
 				case EVectorVMOp::update_id:
 				{
-					serializeIns(1, 2);
+					VVMSer_instruction(1, 2);
 					uint32 DataSetIdx     = VecIndices[0];
 					check(DataSetIdx < (uint32)VVMState->DataSets.Num());
 					FDataSetMeta *DataSet = &VVMState->DataSets[DataSetIdx];
@@ -1459,7 +1464,7 @@ OpCodeSwitch: //I think computed gotos would be a huge win here... maybe write t
 					{
 						int SanityCount = 0;
 						do {
-							int OldMaxID = *DataSet->MaxUsedID;
+							int OldMaxID = FPlatformAtomics::AtomicRead(DataSet->MaxUsedID);
 							if (MaxID <= OldMaxID)
 							{
 								break;
@@ -1470,14 +1475,14 @@ OpCodeSwitch: //I think computed gotos would be a huge win here... maybe write t
 								break;
 							}
 						} while (SanityCount++ < (1 << 30));
-						check(SanityCount < (1 << 30) - 1);
+						VVMDebugBreakIf(SanityCount > (1 << 30) - 1);
 					}
 					InsPtr += 6;
 				}
 				break;
 				case EVectorVMOp::acquire_id:
 				{
-					serializeIns(1, 2);
+					VVMSer_instruction(1, 2);
 					uint32 DataSetIdx = VecIndices[0];
 					check(DataSetIdx < (uint32)VVMState->DataSets.Num());
 					FDataSetMeta *DataSet = &VVMState->DataSets[DataSetIdx];
@@ -1502,10 +1507,11 @@ OpCodeSwitch: //I think computed gotos would be a huge win here... maybe write t
 								break;
 							}
 						} while (SanityCount++ < (1 << 30));
-						check(SanityCount < (1 << 30) - 1);
+						VVMDebugBreakIf(SanityCount > (1 << 30) - 1);
 					}
 					{ //2. append the IDs we acquired in step 1 to the end of the free table array, representing spawned IDs
-						//FreeID table is write-only as far as this invocation of the VM is concerned.
+						//FreeID table is write-only as far as this invocation of the VM is concerned, so the interlocked add w/o filling
+						//in the data is fine
 						int StartNumSpawned = FPlatformAtomics::InterlockedAdd(DataSet->NumSpawnedIDs, NumInstancesThisChunk) + NumInstancesThisChunk;
 						check(StartNumSpawned <= DataSet->FreeIDTable->Max());
 						VVMMemCpy(DataSet->FreeIDTable->GetData() + DataSet->FreeIDTable->Max() - StartNumSpawned, VecReg[1], sizeof(int32) * NumInstancesThisChunk);
diff --git a/Engine/Source/Runtime/VectorVM/Private/VectorVMExperimental_Optimizer.inl b/Engine/Source/Runtime/VectorVM/Private/VectorVMExperimental_Optimizer.inl
--- a/Engine/Source/Runtime/VectorVM/Private/VectorVMExperimental_Optimizer.inl
+++ b/Engine/Source/Runtime/VectorVM/Private/VectorVMExperimental_Optimizer.inl
@@ -716,9 +716,9 @@ VECTORVM_API uint32 OptimizeVectorVMScript(const uint8 *InBytecode, int InByteco
 		EVectorVMOp op		= (EVectorVMOp)*OpPtrIn;
 
 		FVectorVMOptimizeInstruction *Instruction = OptContext->Intermediate.Instructions + OptContext->Intermediate.NumInstructions;
-		Instruction->Index = OptContext->Intermediate.NumInstructions++;
-		Instruction->OpCode = op;
-		Instruction->OpCat = GetOpCategoryFromOp(op);
+		Instruction->Index                   = OptContext->Intermediate.NumInstructions++;
+		Instruction->OpCode                  = op;
+		Instruction->OpCat                   = GetOpCategoryFromOp(op);
 		Instruction->PtrOffsetInOrigBytecode = (uint32)(OpPtrIn - InBytecode);
 		OpPtrIn++;
 		switch (op)
@@ -1153,9 +1153,7 @@ VECTORVM_API uint32 OptimizeVectorVMScript(const uint8 *InBytecode, int InByteco
 		for (uint32 i = 0; i < OptContext->Intermediate.NumInstructions; ++i)
 		{
 			FVectorVMOptimizeInstruction *InputIns = OptContext->Intermediate.Instructions + i;
-			if (InputIns->OpCode == EVectorVMOp::inputdata_float || 
-				InputIns->OpCode == EVectorVMOp::inputdata_int32 || 
-				InputIns->OpCode == EVectorVMOp::inputdata_half) //noadvance ops can't fuse
+			if (InputIns->OpCode == EVectorVMOp::inputdata_float || InputIns->OpCode == EVectorVMOp::inputdata_int32) //half and no advance ops can't fuse
 			{ 
 				OnePastLastInputIdx = i + 1;
 				bool InputOpCanFuse = true;
@@ -1174,8 +1172,10 @@ VECTORVM_API uint32 OptimizeVectorVMScript(const uint8 *InBytecode, int InByteco
 						case EVectorVMOpCategory::Output:
 							if (OptContext->Intermediate.SSARegisterUsageBuffer[InputIns->Input.DstRegPtrOffset] == OptContext->Intermediate.SSARegisterUsageBuffer[Ins_j->Output.RegPtrOffset + 1])
 							{
-								++InputIns->Input.FuseCount;
-								Ins_j->Output.CopyFromInputInsIdx = i;
+								if (Ins_j->OpCode != EVectorVMOp::outputdata_half) {
+									++InputIns->Input.FuseCount;
+									Ins_j->Output.CopyFromInputInsIdx = i;
+								}
 							}
 							break;
 						case EVectorVMOpCategory::Op: {
@@ -1285,8 +1285,7 @@ VECTORVM_API uint32 OptimizeVectorVMScript(const uint8 *InBytecode, int InByteco
 		}
 	}
 
-	if (1)
-	{ //Step 7: remove instructions where outputs are never used 
+	if (1) { //Step 7: remove instructions where outputs are never used 
 		int NumRemovedInstructions = 0;
 		FVectorVMOptimizeInsRegUsage RegUsage;
 		FVectorVMOptimizeInsRegUsage RegUsage2;
@@ -1357,52 +1356,56 @@ VECTORVM_API uint32 OptimizeVectorVMScript(const uint8 *InBytecode, int InByteco
 					for (int j = 0; j < OutputInsRegUse.NumOutputRegisters; ++j)
 					{
 						uint16 OutReg = OptContext->Intermediate.RegisterUsageBuffer[OutputInsRegUse.RegIndices[OutputInsRegUse.NumInputRegisters + j]];
-						OptContext->Intermediate.SSARegisterUsageBuffer[OutputInsRegUse.RegIndices[OutputInsRegUse.NumInputRegisters + j]] = SSARegCount;
-						int LastUsedAsInputInsIdx = -1;
-						for (uint32 k = i + 1; k < OptContext->Intermediate.NumInstructions; ++k)
-						{
-							FVectorVMOptimizeInstruction *InputIns = OptContext->Intermediate.Instructions + k;
-							GetRegistersUsedForInstruction(OptContext, InputIns, &InputInsRegUse);
-							for (int ii = 0; ii < InputInsRegUse.NumOutputRegisters; ++ii)
+						if (OutReg != 0xFFFF) {
+							OptContext->Intermediate.SSARegisterUsageBuffer[OutputInsRegUse.RegIndices[OutputInsRegUse.NumInputRegisters + j]] = SSARegCount;
+							int LastUsedAsInputInsIdx = -1;
+							for (uint32 k = i + 1; k < OptContext->Intermediate.NumInstructions; ++k)
 							{
-								if (OptContext->Intermediate.RegisterUsageBuffer[InputInsRegUse.RegIndices[InputInsRegUse.NumInputRegisters + ii]] == OutReg)
+								FVectorVMOptimizeInstruction *InputIns = OptContext->Intermediate.Instructions + k;
+								GetRegistersUsedForInstruction(OptContext, InputIns, &InputInsRegUse);
+								for (int ii = 0; ii < InputInsRegUse.NumOutputRegisters; ++ii)
 								{
-									//this register is overwritten, we need to generate a new register
-									++SSARegCount;
-									check(SSARegCount <= (int)OptContext->Intermediate.NumRegistersUsed);
-									goto DoneThisOutput2;
+									if (OptContext->Intermediate.RegisterUsageBuffer[InputInsRegUse.RegIndices[InputInsRegUse.NumInputRegisters + ii]] == OutReg)
+									{
+										//this register is overwritten, we need to generate a new register
+										++SSARegCount;
+										check(SSARegCount <= (int)OptContext->Intermediate.NumRegistersUsed);
+										goto DoneThisOutput2;
+									}
 								}
-							}
-							for (int ii = 0; ii < InputInsRegUse.NumInputRegisters; ++ii)
-							{
-								if (OptContext->Intermediate.RegisterUsageBuffer[InputInsRegUse.RegIndices[ii]] == OutReg)
+								for (int ii = 0; ii < InputInsRegUse.NumInputRegisters; ++ii)
 								{
-									if (InputIns->OpCat == EVectorVMOpCategory::Output)
+									if (OptContext->Intermediate.RegisterUsageBuffer[InputInsRegUse.RegIndices[ii]] == OutReg)
 									{
-										if (InputIns->OpCode != EVectorVMOp::copy_to_output)
+										if (InputIns->OpCat == EVectorVMOpCategory::Output)
 										{
-											if (OutputIns->OpCode == EVectorVMOp::acquireindex)
-											{
-												OptContext->Intermediate.SSARegisterUsageBuffer[InputIns->Output.RegPtrOffset] = SSARegCount;
-											}
-											else
+											if (InputIns->OpCode != EVectorVMOp::copy_to_output)
 											{
-												OptContext->Intermediate.SSARegisterUsageBuffer[InputIns->Output.RegPtrOffset + 1] = SSARegCount;
-												LastUsedAsInputInsIdx = k;
+												if (OutputIns->OpCode == EVectorVMOp::acquireindex)
+												{
+													OptContext->Intermediate.SSARegisterUsageBuffer[InputIns->Output.RegPtrOffset] = SSARegCount;
+												}
+												else
+												{
+													OptContext->Intermediate.SSARegisterUsageBuffer[InputIns->Output.RegPtrOffset + 1] = SSARegCount;
+													LastUsedAsInputInsIdx = k;
+												}
 											}
 										}
-									}
-									else
-									{
-										LastUsedAsInputInsIdx = k;
-										OptContext->Intermediate.SSARegisterUsageBuffer[InputInsRegUse.RegIndices[ii]] = SSARegCount;
+										else
+										{
+											LastUsedAsInputInsIdx = k;
+											OptContext->Intermediate.SSARegisterUsageBuffer[InputInsRegUse.RegIndices[ii]] = SSARegCount;
+										}
 									}
 								}
 							}
-						}
-						if (LastUsedAsInputInsIdx != -1)
-						{
-							++SSARegCount;
+							if (LastUsedAsInputInsIdx != -1)
+							{
+								++SSARegCount;
+							}
+						} else {
+							OptContext->Intermediate.SSARegisterUsageBuffer[OutputInsRegUse.RegIndices[OutputInsRegUse.NumInputRegisters + j]] = 0xFFFF;
 						}
 						DoneThisOutput2: ;
 					}
diff --git a/Engine/Source/Runtime/VectorVM/Private/VectorVMExperimental_Serialization.inl b/Engine/Source/Runtime/VectorVM/Private/VectorVMExperimental_Serialization.inl
--- a/Engine/Source/Runtime/VectorVM/Private/VectorVMExperimental_Serialization.inl
+++ b/Engine/Source/Runtime/VectorVM/Private/VectorVMExperimental_Serialization.inl
@@ -116,8 +116,25 @@ static uint32 VectorVMSerializeSetError_(FVectorVMSerializeState *SerializeState
 #define VVMSer_insEndDecodeExp(...)
 #define VVMSer_insEndExp(...)
 #define VVMSer_initSerializationState(...)
+#define VVMSer_instruction(...)
+#define VVMSer_regUsed(...)
 #else //VVM_SERIALIZE_NO_WRITE
 
+#define VVMSer_instruction(Type, NumParams) if (SerializeState)                                                                         \
+                                            {                                                                                           \
+                                                for (int vi = 0; vi <= (int)(NumParams); ++vi)                                          \
+                                                {                                                                                       \
+                                                    if ((VecIndices[vi] & 0x8000) == 0)											        \
+                                                    {                                                                                   \
+                                                        SerializeState->TempRegFlags[VecIndices[vi]] = VVMRegFlag_Clean + (Type);       \
+                                                    }                                                                                   \
+                                                }                                                                                       \
+                                            }
+#define VVMSer_regUsed(RegIdx, Type)        if (SerializeState && RegIdx != 0xFFFF)                                                     \
+                                            {                                                                                           \
+                                                SerializeState->TempRegFlags[RegIdx] = VVMRegFlag_Clean + (Type);                       \
+                                            }
+
 
 #define VVMSer_chunkStartExp(SerializeState, ChunkIdx_, BatchIdx)                   \
 	FVectorVMSerializeChunk *SerializeChunk = nullptr;                              \
@@ -853,6 +870,9 @@ VECTORVM_API uint32 SerializeVectorVMOutputDataSets(FVectorVMSerializeState *Ser
 #define VVMSer_insEndDecodeExp(...)
 #define VVMSer_insEndExp(...)
 
+#define VVMSer_instruction(...)
+#define VVMSer_regUsed(...)
+
 #define VVMSer_initSerializationState(...)
 #define FreeVectorVMSerializeState(...)
 
diff --git a/Engine/Source/Runtime/VectorVM/Public/VectorVM.h b/Engine/Source/Runtime/VectorVM/Public/VectorVM.h
--- a/Engine/Source/Runtime/VectorVM/Public/VectorVM.h
+++ b/Engine/Source/Runtime/VectorVM/Public/VectorVM.h
@@ -4,319 +4,12 @@
 
 #include "CoreMinimal.h"
 
-#include <atomic>
-
 #include "UObject/ObjectMacros.h"
 #include "Math/RandomStream.h"
 #include "Misc/ByteSwap.h"
 #include "Templates/AlignmentTemplates.h"
-#include "VectorVMExperimental.h"  
-
-//TODO: move to a per platform header and have VM scale vectorization according to vector width.
-#define VECTOR_WIDTH (128)
-#define VECTOR_WIDTH_BYTES (16)
-#define VECTOR_WIDTH_FLOATS (4)
-
-DECLARE_DELEGATE_OneParam(FVMExternalFunction, class FVectorVMExternalFunctionContext& /*Context*/);
-
-UENUM()
-enum class EVectorVMBaseTypes : uint8
-{
-	Float,
-	Int,
-	Bool,
-	Num UMETA(Hidden),
-};
-
-UENUM()
-enum class EVectorVMOperandLocation : uint8
-{
-	Register,
-	Constant,
-	Num
-};
-
-UENUM()
-enum class EVectorVMOp : uint8
-{
-	done,
-	add,
-	sub,
-	mul,
-	div,
-	mad,
-	lerp,
-	rcp,
-	rsq,
-	sqrt,
-	neg,
-	abs,
-	exp,
-	exp2,
-	log,
-	log2,
-	sin,
-	cos,
-	tan,
-	asin,
-	acos,
-	atan,
-	atan2,
-	ceil,
-	floor,
-	fmod,
-	frac,
-	trunc,
-	clamp,
-	min,
-	max,
-	pow,
-	round,
-	sign,
-	step,
-	random,
-	noise,
-
-	//Comparison ops.
-	cmplt,
-	cmple,
-	cmpgt,
-	cmpge,
-	cmpeq,
-	cmpneq,
-	select,
-
-// 	easein,  Pretty sure these can be replaced with just a single smoothstep implementation.
-// 	easeinout,
-
-	//Integer ops
-	addi,
-	subi,
-	muli,
-	divi,//SSE Integer division is not implemented as an intrinsic. Will have to do some manual implementation.
-	clampi,
-	mini,
-	maxi,
-	absi,
-	negi,
-	signi,
-	randomi,
-	cmplti,
-	cmplei,
-	cmpgti,
-	cmpgei,
-	cmpeqi,
-	cmpneqi,
-	bit_and,
-	bit_or,
-	bit_xor,
-	bit_not,
-	bit_lshift,
-	bit_rshift,
-
-	//"Boolean" ops. Currently handling bools as integers.
-	logic_and,
-	logic_or,
-	logic_xor,
-	logic_not,
-
-	//conversions
-	f2i,
-	i2f,
-	f2b,
-	b2f,
-	i2b,
-	b2i,
-
-	// data read/write
-	inputdata_float,
-	inputdata_int32,
-	inputdata_half,
-	inputdata_noadvance_float,
-	inputdata_noadvance_int32,
-	inputdata_noadvance_half,
-	outputdata_float,
-	outputdata_int32,
-	outputdata_half,
-	acquireindex,
-
-	external_func_call,
-
-	/** Returns the index of each instance in the current execution context. */
-	exec_index,
-
-	noise2D,
-	noise3D,
-
-	/** Utility ops for hooking into the stats system for performance analysis. */
-	enter_stat_scope,
-	exit_stat_scope,
-
-	//updates an ID in the ID table
-	update_id,
-	//acquires a new ID from the free list.
-	acquire_id,
-
-	//experimental VM only
-	fused_input1_1, /* 95    op has 1 input operand	   - binary 1          ie: register 0 is an input          (ORDER IS CRUCIAL) */
-	fused_input2_1, /* 96    op has 2 input operands   - binary 01         ie: register 0 is an input          (ORDER IS CRUCIAL) */
-	fused_input2_2, /* 97    op has 2 input operands   - binary 10         ie: register 1 is an input          (ORDER IS CRUCIAL) */
-	fused_input2_3, /* 97    op has 2 input operands   - binary 11         ie: register 1 and 2 are inputs     (ORDER IS CRUCIAL) */
-	fused_input3_1, /* 98    op has 3 input operands   - binary 001        ie: register 1 is an input          (ORDER IS CRUCIAL) */
-	fused_input3_2, /* 99    op has 3 input operands   - binary 010        ie: register 2 is an input          (ORDER IS CRUCIAL) */
-	fused_input3_4, /* 101   op has 3 input operands   - binary 100        ie: register 3 is an input          (ORDER IS CRUCIAL) */
-	fused_input3_3, /* 100   op has 3 input operands   - binary 011        ie: register 1 and 2 are inputs     (ORDER IS CRUCIAL) */
-	fused_input3_5, /* 102   op has 3 input operands   - binary 101        ie: register 3 and 1 are inputs     (ORDER IS CRUCIAL) */
-	fused_input3_6, /* 103   op has 3 input operands   - binary 110        ie: register 2 and 2 are inputs     (ORDER IS CRUCIAL) */
-	fused_input3_7, /* 104   op has 3 input operands   - binary 111        ie: register 1, 2 and 3 are inputs  (ORDER IS CRUCIAL) */
-	copy_to_output, /* 105 */
-	output_batch2,  /* 106 */
-	output_batch3,  /* 107 */
-	output_batch4,  /* 108 */
-	output_batch7,  /* 108 */
-	output_batch8,  /* 108 */
-
-	NumOpcodes
-};
-	
-enum class EVectorVMOpCategory : uint8 {
-	Input,
-	Output,
-	Op,
-	ExtFnCall,
-	IndexGen,
-	ExecIndex,
-	RWBuffer,
-	Stat,
-	Fused,
-	Other
-};
-
-#if STATS
-struct FVMCycleCounter
-{
-	int32 ScopeIndex;
-	uint64 ScopeEnterCycles;
-};
-
-struct FStatScopeData
-{
-	TStatId StatId;
-	std::atomic<uint64> ExecutionCycleCount;
-
-	FStatScopeData(TStatId InStatId) : StatId(InStatId)
-	{
-		ExecutionCycleCount.store(0);
-	}
-	
-	FStatScopeData(const FStatScopeData& InObj)
-	{
-		StatId = InObj.StatId;
-		ExecutionCycleCount.store(InObj.ExecutionCycleCount.load());
-	}
-};
-
-struct FStatStackEntry
-{
-	FCycleCounter CycleCounter;
-	FVMCycleCounter VmCycleCounter;
-};
-#endif
-
-//TODO: 
-//All of this stuff can be handled by the VM compiler rather than dirtying the VM code.
-//Some require RWBuffer like support.
-struct FDataSetMeta
-{
-	TArrayView<uint8 const* RESTRICT const> InputRegisters;
-	TArrayView<uint8 const* RESTRICT const> OutputRegisters;
-
-	uint32 InputRegisterTypeOffsets[3];
-	uint32 OutputRegisterTypeOffsets[3];
-
-	int32 DataSetAccessIndex;	// index for individual elements of this set
-
-	int32 InstanceOffset;		// offset of the first instance processed 
-	
-	TArray<int32>*RESTRICT IDTable;
-	TArray<int32>*RESTRICT FreeIDTable;
-	TArray<int32>*RESTRICT SpawnedIDsTable;
-
-	/** Number of free IDs in the FreeIDTable */
-	int32* NumFreeIDs;
-
-	/** MaxID used in this execution. */
-	int32* MaxUsedID;
-
-	int32 *NumSpawnedIDs;
-
-	int32 IDAcquireTag;
-
-	//Temporary lock we're using for thread safety when writing to the FreeIDTable.
-	//TODO: A lock free algorithm is possible here. We can create a specialized lock free list and reuse the IDTable slots for FreeIndices as Next pointers for our LFL.
-	//This would also work well on the GPU. 
-	//UE-65856 for tracking this work.
-
-	//@NOTE (smcgrath): the lock/unlock functions can go with the ne VM, we don't need them
-#ifndef NIAGARA_EXP_VM
-	FCriticalSection FreeTableLock;
-#endif
-	FORCEINLINE void LockFreeTable();
-	FORCEINLINE void UnlockFreeTable();
-
-	FDataSetMeta()
-		: InputRegisterTypeOffsets{}
-		, OutputRegisterTypeOffsets{}
-		, DataSetAccessIndex(INDEX_NONE)
-		, InstanceOffset(INDEX_NONE)
-		, IDTable(nullptr)
-		, FreeIDTable(nullptr)
-		, SpawnedIDsTable(nullptr)
-		, NumFreeIDs(nullptr)
-		, MaxUsedID(nullptr)
-		, IDAcquireTag(INDEX_NONE)
-	{
-	}	
-
-	FORCEINLINE void Reset()
-	{
-		InputRegisters = TArrayView<uint8 const* RESTRICT const>();
-		OutputRegisters = TArrayView<uint8 const* RESTRICT const>();
-		DataSetAccessIndex = INDEX_NONE;
-		InstanceOffset = INDEX_NONE;
-		IDTable = nullptr;
-		FreeIDTable = nullptr;
-		SpawnedIDsTable = nullptr;
-		NumFreeIDs = nullptr;
-		MaxUsedID = nullptr;
-		IDAcquireTag = INDEX_NONE;
-	}
-
-	FORCEINLINE void Init(const TArrayView<uint8 const* RESTRICT const>& InInputRegisters, const TArrayView<uint8 const* RESTRICT const>& InOutputRegisters, int32 InInstanceOffset, TArray<int32>* InIDTable, TArray<int32>* InFreeIDTable, int32* InNumFreeIDs, int32 *InNumSpawnedIDs, int32* InMaxUsedID, int32 InIDAcquireTag, TArray<int32>* InSpawnedIDsTable)
-	{
-		InputRegisters = InInputRegisters;
-		OutputRegisters = InOutputRegisters;
-
-		DataSetAccessIndex = INDEX_NONE;
-		InstanceOffset = InInstanceOffset;
-		IDTable = InIDTable;
-		FreeIDTable = InFreeIDTable;
-		NumFreeIDs = InNumFreeIDs;
-		NumSpawnedIDs = InNumSpawnedIDs;
-		MaxUsedID = InMaxUsedID;
-		IDAcquireTag = InIDAcquireTag;
-		SpawnedIDsTable = InSpawnedIDsTable;
-	}
-
-private:
-	// Non-copyable and non-movable
-	FDataSetMeta(FDataSetMeta&&) = delete;
-	FDataSetMeta(const FDataSetMeta&) = delete;
-	FDataSetMeta& operator=(FDataSetMeta&&) = delete;
-	FDataSetMeta& operator=(const FDataSetMeta&) = delete;
-};
-
-
+#include "VectorVMCommon.h"
+#include "VectorVMExperimental.h"
 
 //Data the VM will keep on each dataset locally per thread which is then thread safely pushed to it's destination at the end of execution.
 struct FDataSetThreadLocalTempData
@@ -707,7 +400,11 @@ namespace VectorVM
 		}
 
 		FORCEINLINE bool IsConstant()const { return !IsRegister(); }
+#ifdef NIAGARA_EXP_VM
+		FORCEINLINE bool IsRegister() const { return (bool)AdvanceOffset; }
+#else
 		FORCEINLINE bool IsRegister()const { return (InputOffset & VVM_EXT_FUNC_INPUT_LOC_BIT) != 0; }
+#endif
 		FORCEINLINE int32 GetOffset()const { return InputOffset & VVM_EXT_FUNC_INPUT_LOC_MASK; }
 
 		FORCEINLINE const T Get() { return *InputPtr; }
diff --git a/Engine/Source/Runtime/VectorVM/Public/VectorVMCommon.h b/Engine/Source/Runtime/VectorVM/Public/VectorVMCommon.h
new file mode 100644
--- /dev/null
+++ b/Engine/Source/Runtime/VectorVM/Public/VectorVMCommon.h
@@ -0,0 +1,310 @@
+// Copyright Epic Games, Inc. All Rights Reserved.
+
+#pragma once
+
+#include "CoreMinimal.h"
+#include <atomic>
+
+//TODO: move to a per platform header and have VM scale vectorization according to vector width.
+#define VECTOR_WIDTH (128)
+#define VECTOR_WIDTH_BYTES (16)
+#define VECTOR_WIDTH_FLOATS (4)
+
+DECLARE_DELEGATE_OneParam(FVMExternalFunction, class FVectorVMExternalFunctionContext& /*Context*/);
+
+UENUM()
+enum class EVectorVMBaseTypes : uint8
+{
+	Float,
+	Int,
+	Bool,
+	Num UMETA(Hidden),
+};
+
+UENUM()
+enum class EVectorVMOperandLocation : uint8
+{
+	Register,
+	Constant,
+	Num
+};
+
+UENUM()
+enum class EVectorVMOp : uint8
+{
+	done,
+	add,
+	sub,
+	mul,
+	div,
+	mad,
+	lerp,
+	rcp,
+	rsq,
+	sqrt,
+	neg,
+	abs,
+	exp,
+	exp2,
+	log,
+	log2,
+	sin,
+	cos,
+	tan,
+	asin,
+	acos,
+	atan,
+	atan2,
+	ceil,
+	floor,
+	fmod,
+	frac,
+	trunc,
+	clamp,
+	min,
+	max,
+	pow,
+	round,
+	sign,
+	step,
+	random,
+	noise,
+
+	//Comparison ops.
+	cmplt,
+	cmple,
+	cmpgt,
+	cmpge,
+	cmpeq,
+	cmpneq,
+	select,
+
+// 	easein,  Pretty sure these can be replaced with just a single smoothstep implementation.
+// 	easeinout,
+
+	//Integer ops
+	addi,
+	subi,
+	muli,
+	divi,//SSE Integer division is not implemented as an intrinsic. Will have to do some manual implementation.
+	clampi,
+	mini,
+	maxi,
+	absi,
+	negi,
+	signi,
+	randomi,
+	cmplti,
+	cmplei,
+	cmpgti,
+	cmpgei,
+	cmpeqi,
+	cmpneqi,
+	bit_and,
+	bit_or,
+	bit_xor,
+	bit_not,
+	bit_lshift,
+	bit_rshift,
+
+	//"Boolean" ops. Currently handling bools as integers.
+	logic_and,
+	logic_or,
+	logic_xor,
+	logic_not,
+
+	//conversions
+	f2i,
+	i2f,
+	f2b,
+	b2f,
+	i2b,
+	b2i,
+
+	// data read/write
+	inputdata_float,
+	inputdata_int32,
+	inputdata_half,
+	inputdata_noadvance_float,
+	inputdata_noadvance_int32,
+	inputdata_noadvance_half,
+	outputdata_float,
+	outputdata_int32,
+	outputdata_half,
+	acquireindex,
+
+	external_func_call,
+
+	/** Returns the index of each instance in the current execution context. */
+	exec_index,
+
+	noise2D,
+	noise3D,
+
+	/** Utility ops for hooking into the stats system for performance analysis. */
+	enter_stat_scope,
+	exit_stat_scope,
+
+	//updates an ID in the ID table
+	update_id,
+	//acquires a new ID from the free list.
+	acquire_id,
+
+	//experimental VM only
+	fused_input1_1, /* 95    op has 1 input operand	   - binary 1          ie: register 0 is an input          (ORDER IS CRUCIAL) */
+	fused_input2_1, /* 96    op has 2 input operands   - binary 01         ie: register 0 is an input          (ORDER IS CRUCIAL) */
+	fused_input2_2, /* 97    op has 2 input operands   - binary 10         ie: register 1 is an input          (ORDER IS CRUCIAL) */
+	fused_input2_3, /* 97    op has 2 input operands   - binary 11         ie: register 1 and 2 are inputs     (ORDER IS CRUCIAL) */
+	fused_input3_1, /* 98    op has 3 input operands   - binary 001        ie: register 1 is an input          (ORDER IS CRUCIAL) */
+	fused_input3_2, /* 99    op has 3 input operands   - binary 010        ie: register 2 is an input          (ORDER IS CRUCIAL) */
+	fused_input3_4, /* 101   op has 3 input operands   - binary 100        ie: register 3 is an input          (ORDER IS CRUCIAL) */
+	fused_input3_3, /* 100   op has 3 input operands   - binary 011        ie: register 1 and 2 are inputs     (ORDER IS CRUCIAL) */
+	fused_input3_5, /* 102   op has 3 input operands   - binary 101        ie: register 3 and 1 are inputs     (ORDER IS CRUCIAL) */
+	fused_input3_6, /* 103   op has 3 input operands   - binary 110        ie: register 2 and 2 are inputs     (ORDER IS CRUCIAL) */
+	fused_input3_7, /* 104   op has 3 input operands   - binary 111        ie: register 1, 2 and 3 are inputs  (ORDER IS CRUCIAL) */
+	copy_to_output, /* 105 */
+	output_batch2,  /* 106 */
+	output_batch3,  /* 107 */
+	output_batch4,  /* 108 */
+	output_batch7,  /* 108 */
+	output_batch8,  /* 108 */
+
+	NumOpcodes
+};
+	
+enum class EVectorVMOpCategory : uint8 {
+	Input,
+	Output,
+	Op,
+	ExtFnCall,
+	IndexGen,
+	ExecIndex,
+	RWBuffer,
+	Stat,
+	Fused,
+	Other
+};
+
+#if STATS
+struct FVMCycleCounter
+{
+	int32 ScopeIndex;
+	uint64 ScopeEnterCycles;
+};
+
+struct FStatScopeData
+{
+	TStatId StatId;
+	std::atomic<uint64> ExecutionCycleCount;
+
+	FStatScopeData(TStatId InStatId) : StatId(InStatId)
+	{
+		ExecutionCycleCount.store(0);
+	}
+	
+	FStatScopeData(const FStatScopeData& InObj)
+	{
+		StatId = InObj.StatId;
+		ExecutionCycleCount.store(InObj.ExecutionCycleCount.load());
+	}
+};
+
+struct FStatStackEntry
+{
+	FCycleCounter CycleCounter;
+	FVMCycleCounter VmCycleCounter;
+};
+#endif
+
+//TODO: 
+//All of this stuff can be handled by the VM compiler rather than dirtying the VM code.
+//Some require RWBuffer like support.
+struct FDataSetMeta
+{
+	TArrayView<uint8 const* RESTRICT const> InputRegisters;
+	TArrayView<uint8 const* RESTRICT const> OutputRegisters;
+
+	uint32 InputRegisterTypeOffsets[3];
+	uint32 OutputRegisterTypeOffsets[3];
+
+	int32 DataSetAccessIndex;	// index for individual elements of this set
+
+	int32 InstanceOffset;		// offset of the first instance processed 
+	
+	TArray<int32>*RESTRICT IDTable;
+	TArray<int32>*RESTRICT FreeIDTable;
+	TArray<int32>*RESTRICT SpawnedIDsTable;
+
+	/** Number of free IDs in the FreeIDTable */
+	int32* NumFreeIDs;
+
+	/** MaxID used in this execution. */
+	int32* MaxUsedID;
+
+	int32 *NumSpawnedIDs;
+
+	int32 IDAcquireTag;
+
+	//Temporary lock we're using for thread safety when writing to the FreeIDTable.
+	//TODO: A lock free algorithm is possible here. We can create a specialized lock free list and reuse the IDTable slots for FreeIndices as Next pointers for our LFL.
+	//This would also work well on the GPU. 
+	//UE-65856 for tracking this work.
+
+	//@NOTE (smcgrath): the lock/unlock functions can go with the ne VM, we don't need them
+#ifndef NIAGARA_EXP_VM
+	FCriticalSection FreeTableLock;
+#endif
+	FORCEINLINE void LockFreeTable();
+	FORCEINLINE void UnlockFreeTable();
+
+	FDataSetMeta()
+		: InputRegisterTypeOffsets{}
+		, OutputRegisterTypeOffsets{}
+		, DataSetAccessIndex(INDEX_NONE)
+		, InstanceOffset(INDEX_NONE)
+		, IDTable(nullptr)
+		, FreeIDTable(nullptr)
+		, SpawnedIDsTable(nullptr)
+		, NumFreeIDs(nullptr)
+		, MaxUsedID(nullptr)
+		, IDAcquireTag(INDEX_NONE)
+	{
+	}	
+
+	FORCEINLINE void Reset()
+	{
+		InputRegisters = TArrayView<uint8 const* RESTRICT const>();
+		OutputRegisters = TArrayView<uint8 const* RESTRICT const>();
+		DataSetAccessIndex = INDEX_NONE;
+		InstanceOffset = INDEX_NONE;
+		IDTable = nullptr;
+		FreeIDTable = nullptr;
+		SpawnedIDsTable = nullptr;
+		NumFreeIDs = nullptr;
+		MaxUsedID = nullptr;
+		IDAcquireTag = INDEX_NONE;
+	}
+
+	FORCEINLINE void Init(const TArrayView<uint8 const* RESTRICT const>& InInputRegisters, const TArrayView<uint8 const* RESTRICT const>& InOutputRegisters, int32 InInstanceOffset, TArray<int32>* InIDTable, TArray<int32>* InFreeIDTable, int32* InNumFreeIDs, int32 *InNumSpawnedIDs, int32* InMaxUsedID, int32 InIDAcquireTag, TArray<int32>* InSpawnedIDsTable)
+	{
+		InputRegisters = InInputRegisters;
+		OutputRegisters = InOutputRegisters;
+
+		DataSetAccessIndex = INDEX_NONE;
+		InstanceOffset = InInstanceOffset;
+		IDTable = InIDTable;
+		FreeIDTable = InFreeIDTable;
+		NumFreeIDs = InNumFreeIDs;
+		NumSpawnedIDs = InNumSpawnedIDs;
+		MaxUsedID = InMaxUsedID;
+		IDAcquireTag = InIDAcquireTag;
+		SpawnedIDsTable = InSpawnedIDsTable;
+	}
+
+private:
+	// Non-copyable and non-movable
+	FDataSetMeta(FDataSetMeta&&) = delete;
+	FDataSetMeta(const FDataSetMeta&) = delete;
+	FDataSetMeta& operator=(FDataSetMeta&&) = delete;
+	FDataSetMeta& operator=(const FDataSetMeta&) = delete;
+};
diff --git a/Engine/Source/Runtime/VectorVM/Public/VectorVMExperimental.h b/Engine/Source/Runtime/VectorVM/Public/VectorVMExperimental.h
--- a/Engine/Source/Runtime/VectorVM/Public/VectorVMExperimental.h
+++ b/Engine/Source/Runtime/VectorVM/Public/VectorVMExperimental.h
@@ -287,6 +287,7 @@ struct FVectorVMOptimizeContext
 	uint32                                NumTempRegisters;
 	uint32                                NumExtFns;
 	uint32                                MaxExtFnRegisters;
+	uint32                                NumDummyRegsReq;     //External function "null" registers
 	int32                                 MaxExtFnUsed;
 
 	struct
@@ -361,6 +362,7 @@ struct FVectorVMBatchState
 		struct {
 			FVecReg **        RegData;
 			uint32 *          RegInc;
+			FVecReg *         DummyRegs;
 		} ExtFnDecodedReg;
 		int32 *               RandCounters; //used for external functions only.
 	} ChunkLocalData;
@@ -395,6 +397,7 @@ struct FVectorVMState
 	uint32                      NumOutputDataSets;   //computed in the optimizer
 	uint32                      NumExtFunctions;
 	uint32                      MaxExtFnRegisters;
+	uint32                      NumDummyRegsReq;
 	uint32                      NumUserPtrTable;
 	int32                       TotalNumInstances;
 	volatile int32              NumInstancesAssignedToBatches;
diff --git a/Engine/Source/Runtime/VectorVM/VectorVM.Build.cs b/Engine/Source/Runtime/VectorVM/VectorVM.Build.cs
--- a/Engine/Source/Runtime/VectorVM/VectorVM.Build.cs
+++ b/Engine/Source/Runtime/VectorVM/VectorVM.Build.cs
@@ -25,7 +25,11 @@ public class VectorVM : ModuleRules
                 "Runtime/Engine/Classes/Curves"
             }
         );
-
-
+        PublicDefinitions.AddRange(
+            new string[]
+            {
+                "NIAGARA_EXP_VM",
+            }
+        );
     }
 }
