New VectorVM half float bug fixes #rb none #preflight none

diff --git a/Engine/Source/Runtime/VectorVM/Private/VectorVMExperimental_Optimizer.inl b/Engine/Source/Runtime/VectorVM/Private/VectorVMExperimental_Optimizer.inl
--- a/Engine/Source/Runtime/VectorVM/Private/VectorVMExperimental_Optimizer.inl
+++ b/Engine/Source/Runtime/VectorVM/Private/VectorVMExperimental_Optimizer.inl
@@ -1174,6 +1174,7 @@ VECTORVM_API uint32 OptimizeVectorVMScript(const uint8 *InBytecode, int InByteco
 							{
 								if (Ins_j->OpCode != EVectorVMOp::outputdata_half) {
 									++InputIns->Input.FuseCount;
+									check(InputIns->Index == i);
 									Ins_j->Output.CopyFromInputInsIdx = i;
 								}
 							}
@@ -1658,7 +1659,36 @@ VECTORVM_API uint32 OptimizeVectorVMScript(const uint8 *InBytecode, int InByteco
 		}
 	}
 
-	{ //Step 14: make sure all the copy-to-output instructions are grouped together
+	{ //Step 14: if we re-ordered inputs, the CopyFromInputInsIdx on output instructions could be wrong.  Go fix those
+		for (uint32 OutputInsIdx = 0; OutputInsIdx < OptContext->Intermediate.NumInstructions; ++OutputInsIdx)
+		{
+			FVectorVMOptimizeInstruction *OutputIns = OptContext->Intermediate.Instructions + OutputInsIdx;
+			if (OutputIns->OpCat == EVectorVMOpCategory::Output && OutputIns->Output.CopyFromInputInsIdx != -1)
+			{
+				check(OutputIns->OpCode != EVectorVMOp::outputdata_half);
+				FVectorVMOptimizeInstruction *InputIns = OptContext->Intermediate.Instructions + OutputIns->Output.CopyFromInputInsIdx;
+				if (InputIns->Index != OutputIns->Output.CopyFromInputInsIdx)
+				{
+					for (uint32 i = 0; i < OptContext->Intermediate.NumInstructions; ++i) {
+						FVectorVMOptimizeInstruction *Ins = OptContext->Intermediate.Instructions + i;
+						if (Ins->Index == OutputIns->Output.CopyFromInputInsIdx) {
+							check(Ins->OpCat == EVectorVMOpCategory::Input);
+							check(Ins->OpCode != EVectorVMOp::inputdata_half);
+							uint8 InputRegType  = (uint8)Ins->OpCode - (uint8)EVectorVMOp::inputdata_float;
+							uint8 OutputRegType = (uint8)OutputIns->OpCode - (uint8)EVectorVMOp::outputdata_float;
+							check(InputRegType == OutputRegType);
+							OutputIns->Output.CopyFromInputInsIdx = (int)i;
+							break;
+						}
+					}
+				}
+			}
+		}
+	}
+
+
+
+	{ //Step 15: make sure all the copy-to-output instructions are grouped together
 		int FirstCopyFromInputInsIdx = -1;
 		int LastCopyFromInputInsIdx = -1;
 		//when the copy-to-output instructions get written to the bytecode they'll get grouped into as few instructions as possible
@@ -1691,58 +1721,29 @@ VECTORVM_API uint32 OptimizeVectorVMScript(const uint8 *InBytecode, int InByteco
 
 		
 		if (LastCopyFromInputInsIdx >= FirstCopyFromInputInsIdx + 2)
-		{
-			//fixup CopyFromInputIns for outpus where the corresponding input instruction was moved
-			for (int i = FirstCopyFromInputInsIdx; i <= LastCopyFromInputInsIdx; ++i)
-			{
-				FVectorVMOptimizeInstruction *OutputIns = OptContext->Intermediate.Instructions + i;
-				check(OutputIns->OpCat == EVectorVMOpCategory::Output);
-				check(OutputIns->Output.CopyFromInputInsIdx != -1);
-				FVectorVMOptimizeInstruction *InputIns = OptContext->Intermediate.Instructions + OutputIns->Output.CopyFromInputInsIdx;
-				if (InputIns->Index == OutputIns->Output.CopyFromInputInsIdx)
-				{
-					check(InputIns->OpCat == EVectorVMOpCategory::Input);
-				}
-				else //this input instruction has been re-ordered.  Fix the output instruction so it points to the correct input instruction
-				{ 
-					for (uint32 j = 0; j < OptContext->Intermediate.NumInstructions; ++j)
-					{
-						FVectorVMOptimizeInstruction *Ins = OptContext->Intermediate.Instructions + j;
-						if (Ins->Index == OutputIns->Output.CopyFromInputInsIdx)
-						{
-							check(Ins->OpCat == EVectorVMOpCategory::Input);
-							check(OutputIns->Output.CopyFromInputInsIdx != (int)j); //should have found it above and never hit this loop if it's already correct
-							OutputIns->Output.CopyFromInputInsIdx = (int)j;
-							break;
-						}
-					}
-				}
-			}
-
-			{ // small list, very stupid bubble sort
-				bool sorted = false;
-				int NumCopyInstructions = LastCopyFromInputInsIdx - FirstCopyFromInputInsIdx + 1;
-				while (!sorted) {
-					sorted = true;
-					for (int i = 0; i < NumCopyInstructions - 1; ++i) {
-						FVectorVMOptimizeInstruction *Ins0 = OptContext->Intermediate.Instructions + FirstCopyFromInputInsIdx + i;
-						FVectorVMOptimizeInstruction *Ins1 = Ins0 + 1;
-						uint64 Key0 = VVMCopyToOutputInsGetSortKey(OptContext->Intermediate.Instructions, Ins0);
-						uint64 Key1 = VVMCopyToOutputInsGetSortKey(OptContext->Intermediate.Instructions, Ins1);
-
-						if (Key1 < Key0) {
-							FVectorVMOptimizeInstruction Temp = *Ins0;
-							*Ins0 = *Ins1;
-							*Ins1 = Temp;
-							sorted = false;
-						}
+		{  // small list, very stupid bubble sort
+			bool sorted = false;
+			int NumCopyInstructions = LastCopyFromInputInsIdx - FirstCopyFromInputInsIdx + 1;
+			while (!sorted) {
+				sorted = true;
+				for (int i = 0; i < NumCopyInstructions - 1; ++i) {
+					FVectorVMOptimizeInstruction *Ins0 = OptContext->Intermediate.Instructions + FirstCopyFromInputInsIdx + i;
+					FVectorVMOptimizeInstruction *Ins1 = Ins0 + 1;
+					uint64 Key0 = VVMCopyToOutputInsGetSortKey(OptContext->Intermediate.Instructions, Ins0);
+					uint64 Key1 = VVMCopyToOutputInsGetSortKey(OptContext->Intermediate.Instructions, Ins1);
+
+					if (Key1 < Key0) {
+						FVectorVMOptimizeInstruction Temp = *Ins0;
+						*Ins0 = *Ins1;
+						*Ins1 = Temp;
+						sorted = false;
 					}
 				}
 			}
 		}
 	}
 
-	{ //Step 15: group and sort all regular output instructions
+	{ //Step 16: group and sort all regular output instructions
 		for (uint32 i = 0; i < OptContext->Intermediate.NumInstructions; ++i)
 		{
 			FVectorVMOptimizeInstruction *InsStart = OptContext->Intermediate.Instructions + i;
@@ -1789,7 +1790,7 @@ VECTORVM_API uint32 OptimizeVectorVMScript(const uint8 *InBytecode, int InByteco
 	}
 
 
-	{ //Step 16: correct the register fuse buffer to make sure the instruction indices match re-ordered inputs
+	{ //Step 17: correct the register fuse buffer and output copy input indices to make sure the instruction indices match re-ordered inputs
 		int32 *TempInputRegisterFuseBuffer = (int32 *)OptContext->Init.ReallocFn(nullptr, sizeof(int32) * OptContext->Intermediate.NumRegistersUsed, __FILE__, __LINE__);
 		if (TempInputRegisterFuseBuffer == nullptr)
 		{
@@ -1801,24 +1802,21 @@ VECTORVM_API uint32 OptimizeVectorVMScript(const uint8 *InBytecode, int InByteco
 		for (uint32 InputInsIdx = 0; InputInsIdx < OptContext->Intermediate.NumInstructions; ++InputInsIdx)
 		{
 			FVectorVMOptimizeInstruction *InputIns = OptContext->Intermediate.Instructions + InputInsIdx;
-			if (InputIns->OpCat == EVectorVMOpCategory::Input)
+			if (InputIns->OpCat == EVectorVMOpCategory::Input && InputIns->Index != InputInsIdx) //only worry about instructions that are not in their original place
 			{
-				if (InputIns->Index != InputInsIdx)
-				{ //only worry about instructions that are not in their original place
-					//fixup register fuse buffer
-					for (uint32 i = 0; i < OptContext->Intermediate.NumRegistersUsed; ++i)
+				//fixup register fuse buffer
+				for (uint32 i = 0; i < OptContext->Intermediate.NumRegistersUsed; ++i)
+				{
+					if (TempInputRegisterFuseBuffer[i] == InputIns->Index)
 					{
-						if (TempInputRegisterFuseBuffer[i] == InputIns->Index)
-						{
-							OptContext->Intermediate.InputRegisterFuseBuffer[i] = InputInsIdx;
-						}
+						OptContext->Intermediate.InputRegisterFuseBuffer[i] = InputInsIdx;
 					}
 				}
 			}
 		}	
 	}
 
-	{ //Step 17: use the SSA registers to compute the minimized registers required and write them back into the register usage buffer
+	{ //Step 18: use the SSA registers to compute the minimized registers required and write them back into the register usage buffer
 		int MaxLiveRegisters = 0;
 		uint16 *SSAUseMap = (uint16 *)OptContext->Init.ReallocFn(nullptr, sizeof(uint16) * NumSSARegistersUsed, __FILE__, __LINE__);
 		if (SSAUseMap == nullptr)
@@ -1922,7 +1920,7 @@ VECTORVM_API uint32 OptimizeVectorVMScript(const uint8 *InBytecode, int InByteco
 		OptContext->NumTempRegisters = (uint32)MaxLiveRegisters;
 	}
 
-	{ //Step 18: write the final optimized bytecode
+	{ //Step 19: write the final optimized bytecode
 		//this goes over the instruction list twice.  The first time to figure out how many bytes are required for the bytecode, the second to write the bytecode.
 		uint8 *OptimizedBytecode = nullptr;
 		int NumOptimizedBytesRequired = 0;
