Fixup some warnings generated from usage of VectorRegister2Int64 with new experimental VectorVM

diff --git a/Engine/Source/Runtime/Core/Public/Math/UnrealMathSSE.h b/Engine/Source/Runtime/Core/Public/Math/UnrealMathSSE.h
--- a/Engine/Source/Runtime/Core/Public/Math/UnrealMathSSE.h
+++ b/Engine/Source/Runtime/Core/Public/Math/UnrealMathSSE.h
@@ -76,6 +76,9 @@ typedef __m128	VectorRegister4Float;
 // 4 int32s
 typedef __m128i VectorRegister4Int;
 
+// 2 int64s
+typedef __m128i VectorRegister2Int64;
+
 // 2 doubles
 typedef __m128d	VectorRegister2Double;
 
@@ -380,7 +383,7 @@ FORCEINLINE VectorRegister4Int MakeVectorRegisterInt(int32 X, int32 Y, int32 Z,
 	return _mm_setr_epi32(X, Y, Z, W);
 }
 
-FORCEINLINE VectorRegister4Int MakeVectorRegisterInt64(int64 X, int64 Y) {
+FORCEINLINE VectorRegister2Int64 MakeVectorRegisterInt64(int64 X, int64 Y) {
 	return _mm_set_epi64x(Y, X);
 }
 
diff --git a/Engine/Source/Runtime/VectorVM/Private/VectorVMExperimental.inl b/Engine/Source/Runtime/VectorVM/Private/VectorVMExperimental.inl
--- a/Engine/Source/Runtime/VectorVM/Private/VectorVMExperimental.inl
+++ b/Engine/Source/Runtime/VectorVM/Private/VectorVMExperimental.inl
@@ -1000,7 +1000,7 @@ OpCodeSwitch: //I think computed gotos would be a huge win here... maybe write t
 					{
 						VectorRegister4i RandReg = VVMXorwowStep(BatchState);
 						VectorRegister4i IntPart = VectorIntOr(VectorShiftRightImmLogical(RandReg, 9), VectorIntSet1(0x3F800000));
-						VectorRegister4f FltPart = VectorCastIntToFloat(IntPart);
+						VectorRegister4f FltPart = VectorRegister4f(VectorCastIntToFloat(IntPart));
 						VecReg[1][i].v = VectorMultiply(VectorSubtract(FltPart, VectorSet1(1.f)), VecReg[0][i & RegInc[0]].v);
 					}
 #					if defined(VVM_INCLUDE_SERIALIZATION) && !defined(VVM_SERIALIZE_NO_WRITE)
diff --git a/Engine/Source/Runtime/VectorVM/Public/VectorVMExperimental.h b/Engine/Source/Runtime/VectorVM/Public/VectorVMExperimental.h
--- a/Engine/Source/Runtime/VectorVM/Public/VectorVMExperimental.h
+++ b/Engine/Source/Runtime/VectorVM/Public/VectorVMExperimental.h
@@ -374,7 +374,7 @@ struct FVectorVMBatchState
 	int                       NumInstances;
 
 	VectorRegister4i          RandState[5]; //xorwor state for random/randomi instructions.  DIs use RandomStream.
-	VectorRegister4i          RandCounters;
+	VectorRegister2Int64      RandCounters;
 
 	FRandomStream             RandStream;
 	int                       UseCount;
