Fix for empty scripts (i.e. scripts that only has a single 'done' instruction)

diff --git a/Engine/Source/Runtime/VectorVM/Private/VectorVMExperimental_Optimizer.inl b/Engine/Source/Runtime/VectorVM/Private/VectorVMExperimental_Optimizer.inl
--- a/Engine/Source/Runtime/VectorVM/Private/VectorVMExperimental_Optimizer.inl
+++ b/Engine/Source/Runtime/VectorVM/Private/VectorVMExperimental_Optimizer.inl
@@ -671,6 +671,12 @@ VECTORVM_API uint32 OptimizeVectorVMScript(const uint8 *InBytecode, int InByteco
 		return 0;
 	}
 
+	// skip the script if the input is empty
+	if ((InBytecodeLen == 1) && (EVectorVMOp(InBytecode[0]) == EVectorVMOp::done))
+	{
+		return 0;
+	}
+
 	if (OptContext->Init.ReallocFn == nullptr)
 	{
 		OptContext->Init.ReallocFn = VVMDefaultRealloc;
