New VectorVM bug fixes.
Fixed pointer alignment causing compiler to generate aligned instructions on unaligned data.  Fixed optimization bug on empty scripts.

diff --git a/Engine/Source/Runtime/VectorVM/Private/VectorVMExperimental.inl b/Engine/Source/Runtime/VectorVM/Private/VectorVMExperimental.inl
--- a/Engine/Source/Runtime/VectorVM/Private/VectorVMExperimental.inl
+++ b/Engine/Source/Runtime/VectorVM/Private/VectorVMExperimental.inl
@@ -613,7 +613,7 @@ VECTORVM_API FVectorVMState *InitVectorVMState(FVectorVMInitData *InitData, FVec
 
 	{ //compute the number of overhead bytes for this VVM State
 		size_t NumDataSetOutputBytesRequired            = sizeof(volatile long)            * InitData->OptimizeContext->NumOutputDataSets;
-		size_t NumBatchStateBytesRequired               = sizeof(FVectorVMBatchState)      * NumBatches; //this is *NOT* the memory per batch, just the overhead stored in VVMState
+		size_t NumBatchStateBytesRequired               = sizeof(FVectorVMBatchState)      * NumBatches + 16; //this is *NOT* the memory per batch, just the overhead stored in VVMState
 		size_t NumExtFnBytesRequired                    = sizeof(FVectorVMExtFunctionData) * InitData->ExtFunctionTable.Num();
 		NumVVMStateBytesRequired                        = VVM_ALIGN_64(NumDataSetOutputBytesRequired +
 			                                                           NumExtFnBytesRequired +
@@ -682,6 +682,7 @@ VECTORVM_API FVectorVMState *InitVectorVMState(FVectorVMInitData *InitData, FVec
 		}
 		uint8 *StateAfterExtFnPtr = VVMStatePtr;
 		VVMState->NumOutputPerDataSet  = (volatile int32 *)VVMStatePtr;                  IncVVMStatePtr(sizeof(volatile int32) * InitData->OptimizeContext->NumOutputDataSets, 0);
+		VVMStatePtr = (uint8 *)VVM_ALIGN_16(VVMStatePtr);
 		VVMState->BatchStates          = (FVectorVMBatchState *)VVMStatePtr;             IncVVMStatePtr(sizeof(volatile FVectorVMBatchState) * NumBatches, 0);
 		FMemory::Memset(StateAfterExtFnPtr, 0, (size_t)(VVMStatePtr - StateAfterExtFnPtr));
 	}
@@ -857,6 +858,7 @@ static void ExecVVMBatch(FVectorVMState *VVMState, int ExecIdx, FVectorVMSeriali
 		check(false);
 		return;
 	}
+	check((((size_t)BatchState) & 0xF) == 0); //must be 16 byte aligned
 	if (BatchState->RegisterData == nullptr)
 	{
 		check(BatchIdx != 0); //first batch state should have set the pointers in init()
diff --git a/Engine/Source/Runtime/VectorVM/Private/VectorVMExperimental_Optimizer.inl b/Engine/Source/Runtime/VectorVM/Private/VectorVMExperimental_Optimizer.inl
--- a/Engine/Source/Runtime/VectorVM/Private/VectorVMExperimental_Optimizer.inl
+++ b/Engine/Source/Runtime/VectorVM/Private/VectorVMExperimental_Optimizer.inl
@@ -666,7 +666,7 @@ VECTORVM_API uint32 OptimizeVectorVMScript(const uint8 *InBytecode, int InByteco
 													VVMOptimizeWriteRegIndex(3, 1);		\
 													OpPtrIn += 9;
 	FreeVectorVMOptimizeContext(OptContext);
-	if (InBytecode == nullptr || InBytecodeLen == 0)
+	if (InBytecode == nullptr || InBytecodeLen <= 1)
 	{
 		return 0;
 	}
diff --git a/Engine/Source/Runtime/VectorVM/Public/VectorVMExperimental.h b/Engine/Source/Runtime/VectorVM/Public/VectorVMExperimental.h
--- a/Engine/Source/Runtime/VectorVM/Public/VectorVMExperimental.h
+++ b/Engine/Source/Runtime/VectorVM/Public/VectorVMExperimental.h
@@ -374,12 +374,14 @@ struct FVectorVMBatchState
 	int                       NumInstances;
 
 	VectorRegister4i          RandState[5]; //xorwor state for random/randomi instructions.  DIs use RandomStream.
-	VectorRegister2Int64      RandCounters;
+	VectorRegister4i          RandCounters;
 
 	FRandomStream             RandStream;
 	int                       UseCount;
 };
 
+static_assert((sizeof(FVectorVMBatchState) & 0xF) == 0, "FVectorVMBatchState must be 16 byte aligned");
+
 struct FVectorVMState
 {
 	uint8 *                     Bytecode;
