VectorVM Optimizer bug fix calculating dummy registers. #jira UE-146588 #rb jeff.roberts #preflight 6239feafc73745c600a07f34

diff --git a/Engine/Source/Runtime/VectorVM/Private/VectorVMExperimental_Optimizer.inl b/Engine/Source/Runtime/VectorVM/Private/VectorVMExperimental_Optimizer.inl
--- a/Engine/Source/Runtime/VectorVM/Private/VectorVMExperimental_Optimizer.inl
+++ b/Engine/Source/Runtime/VectorVM/Private/VectorVMExperimental_Optimizer.inl
@@ -890,7 +890,6 @@ VECTORVM_API uint32 OptimizeVectorVMScript(const uint8 *InBytecode, int InByteco
 			break;
 			case EVectorVMOp::external_func_call:
 			{
-				uint32 DummyRegCount = 0;
 				uint8 ExtFnIdx = *OpPtrIn;
 				check(ExtFnIdx < NumExtFns);
 
@@ -905,7 +904,6 @@ VECTORVM_API uint32 OptimizeVectorVMScript(const uint8 *InBytecode, int InByteco
 					if (VecIndices[i] == 0xFFFF)
 					{ //invalid, just write it out
 						OptContext->Intermediate.RegisterUsageBuffer[OptContext->Intermediate.NumRegistersUsed] = 0xFFFF;
-						++DummyRegCount;
 					}
 					else
 					{
@@ -931,14 +929,10 @@ VECTORVM_API uint32 OptimizeVectorVMScript(const uint8 *InBytecode, int InByteco
 					int Idx = ExtFnIOData[ExtFnIdx].NumInputs + i;
 					check((VecIndices[Idx] & 0x8000) == 0 || VecIndices[Idx] == 0xFFFF); //can't output to a const... 0xFFFF is invalid
 					if (VecIndices[Idx] == 0xFFFF) {
-						++DummyRegCount;
 					}
 					OptContext->Intermediate.RegisterUsageBuffer[OptContext->Intermediate.NumRegistersUsed] = VecIndices[Idx];
 					++OptContext->Intermediate.NumRegistersUsed;
 				}
-				if (DummyRegCount > OptContext->NumDummyRegsReq) {
-					OptContext->NumDummyRegsReq = DummyRegCount;
-				}
 				OptContext->MaxExtFnUsed      = VVM_MAX(OptContext->MaxExtFnUsed, ExtFnIdx);
 				OptContext->MaxExtFnRegisters = VVM_MAX(OptContext->MaxExtFnRegisters, (uint32)(ExtFnIOData[ExtFnIdx].NumInputs + ExtFnIOData[ExtFnIdx].NumOutputs));
 				OpPtrIn += 1 + (ExtFnIOData[ExtFnIdx].NumInputs + ExtFnIOData[ExtFnIdx].NumOutputs) * 2;
@@ -2228,14 +2222,22 @@ VECTORVM_API uint32 OptimizeVectorVMScript(const uint8 *InBytecode, int InByteco
 					VVMOptWriteU16(OptContext->Intermediate.RegisterUsageBuffer[Ins->IndexGen.RegPtrOffset + 1]); //2: Write-gather Output Register
 					VVMOptWriteU16(OptContext->Intermediate.RegisterUsageBuffer[Ins->IndexGen.RegPtrOffset + 2]); //3: Original VM Output Register (if 0xFFFF then nothing)
 					break;
-				case EVectorVMOpCategory::ExtFnCall:
+				case EVectorVMOpCategory::ExtFnCall: {
 					VVMOptWriteByte(Ins->OpCode);
 					VVMOptWriteU16(Ins->ExtFnCall.ExtFnIdx);
+					uint32 NumDummyRegs = 0;
 					for (int j = 0; j < ExtFnIOData[Ins->ExtFnCall.ExtFnIdx].NumInputs + ExtFnIOData[Ins->ExtFnCall.ExtFnIdx].NumOutputs; ++j)
 					{
-						VVMOptWriteU16(OptContext->Intermediate.RegisterUsageBuffer[Ins->ExtFnCall.RegPtrOffset + j]);
+						uint16 Reg = OptContext->Intermediate.RegisterUsageBuffer[Ins->ExtFnCall.RegPtrOffset + j];
+						VVMOptWriteU16(Reg);
+						if (Reg == 0xFFFF) {
+							++NumDummyRegs;
+						}
 					}
-					break;
+					if (NumDummyRegs > OptContext->NumDummyRegsReq) {
+						OptContext->NumDummyRegsReq = NumDummyRegs;
+					}
+				} break;
 				case EVectorVMOpCategory::ExecIndex: 
 					VVMOptWriteByte(Ins->OpCode);
 					VVMOptWriteU16(OptContext->Intermediate.RegisterUsageBuffer[Ins->ExecIndex.RegPtrOffset]);
