Fix VectorVM alignment bug on f2i/i2f instructions. #jira UE-146776 #rb rob.krajcarski #preflight none

diff --git a/Engine/Source/Runtime/VectorVM/Private/VectorVMExperimental.inl b/Engine/Source/Runtime/VectorVM/Private/VectorVMExperimental.inl
--- a/Engine/Source/Runtime/VectorVM/Private/VectorVMExperimental.inl
+++ b/Engine/Source/Runtime/VectorVM/Private/VectorVMExperimental.inl
@@ -1101,7 +1101,8 @@ OpCodeSwitch: //I think computed gotos would be a huge win here... maybe write t
 				case EVectorVMOp::f2i:
 					for (int i = 0; i < NumLoops; ++i)
 					{
-						VecReg[1][i].i = VectorFloatToInt(VecReg[0][i & RegInc[0]].v);
+						VectorRegister4 r0 = VectorLoad(&VecReg[0][i & RegInc[0]].v);
+						VectorIntStoreAligned(VectorFloatToInt(r0), &VecReg[1][i].i);
 					}
 					VVMSer_regUsed(VecIndices[1], 1);
 					InsPtr += 4;
@@ -1109,7 +1110,8 @@ OpCodeSwitch: //I think computed gotos would be a huge win here... maybe write t
 				case EVectorVMOp::i2f:
 					for (int i = 0; i < NumLoops; ++i)
 					{
-						VecReg[1][i].v = VectorIntToFloat(VecReg[0][i & RegInc[0]].i);
+						VectorRegister4i r0 = VectorIntLoad(&VecReg[0][i & RegInc[0]].i);
+						VectorStoreAligned(VectorIntToFloat(r0), &VecReg[1][i].v);
 					}
 					VVMSer_regUsed(VecIndices[1], 0);
 					InsPtr += 4;
