VectorVM neon compiler fix. VectorRegister4 vs VectorRegister4f #rb none #preflight none

diff --git a/Engine/Source/Runtime/VectorVM/Private/VectorVMExperimental.inl b/Engine/Source/Runtime/VectorVM/Private/VectorVMExperimental.inl
--- a/Engine/Source/Runtime/VectorVM/Private/VectorVMExperimental.inl
+++ b/Engine/Source/Runtime/VectorVM/Private/VectorVMExperimental.inl
@@ -1101,7 +1101,7 @@ OpCodeSwitch: //I think computed gotos would be a huge win here... maybe write t
 				case EVectorVMOp::f2i:
 					for (int i = 0; i < NumLoops; ++i)
 					{
-						VectorRegister4 r0 = VectorLoad(&VecReg[0][i & RegInc[0]].v);
+						VectorRegister4f r0 = VectorLoad(&VecReg[0][i & RegInc[0]].v);
 						VectorIntStoreAligned(VectorFloatToInt(r0), &VecReg[1][i].i);
 					}
 					VVMSer_regUsed(VecIndices[1], 1);
