Vector VM optimizer: temporarily removed redundant instructions that's causing an SSA allocation bug.  Will investigate bug in more depth, but this temp fix solves the Lyra (and potentially other) bugs. #jira UE-146776 #rb rob.krajcarski #preflight none

diff --git a/Engine/Source/Runtime/VectorVM/Private/VectorVMExperimental_Optimizer.inl b/Engine/Source/Runtime/VectorVM/Private/VectorVMExperimental_Optimizer.inl
--- a/Engine/Source/Runtime/VectorVM/Private/VectorVMExperimental_Optimizer.inl
+++ b/Engine/Source/Runtime/VectorVM/Private/VectorVMExperimental_Optimizer.inl
@@ -1286,7 +1286,7 @@ VECTORVM_API uint32 OptimizeVectorVMScript(const uint8 *InBytecode, int InByteco
 		}
 	}
 
-	if (1) { //Step 7: remove instructions where outputs are never used 
+	if (0) { //Step 7: remove instructions where outputs are never used 
 		int NumRemovedInstructions = 0;
 		FVectorVMOptimizeInsRegUsage RegUsage;
 		FVectorVMOptimizeInsRegUsage RegUsage2;
@@ -1302,7 +1302,6 @@ VECTORVM_API uint32 OptimizeVectorVMScript(const uint8 *InBytecode, int InByteco
 				if (Ins->OpCat == EVectorVMOpCategory::Op && !(Ins->OpCode == EVectorVMOp::random || Ins->OpCode == EVectorVMOp::randomi))
 				{
 					//can we remove random instructions? I dunno! so lets not for now
-					bool InsRequired = false;
 					GetRegistersUsedForInstruction(OptContext, Ins, &RegUsage);
 					for (int OutputIdx = 0; OutputIdx < RegUsage.NumOutputRegisters; ++OutputIdx)
 					{
@@ -1315,20 +1314,19 @@ VECTORVM_API uint32 OptimizeVectorVMScript(const uint8 *InBytecode, int InByteco
 								uint16 RegIdx2 = OptContext->Intermediate.SSARegisterUsageBuffer[RegUsage2.RegIndices[k]];
 								if (RegIdx == RegIdx2)
 								{
-									InsRequired = true;
-									break;
+									goto InstructionRequired;
 								}
 							}
 						}
 					}
-					if (!InsRequired)
-					{
+					{ //instruction isn't required
 						FMemory::Memmove(OptContext->Intermediate.Instructions + i, OptContext->Intermediate.Instructions + i + 1, sizeof(FVectorVMOptimizeInstruction) * (OptContext->Intermediate.NumInstructions - i - 1));
 						++NumRemovedInstructionsThisTime;
 						++NumRemovedInstructions;
 						--OptContext->Intermediate.NumInstructions;
 						--i;
 					}
+					InstructionRequired: ;
 				}
 			}
 			if (++SanityCount >= 16384)
