New VectorVM optimizer fixes.  Fixes some invalid instruction removal bugs as well as incorrect SSA calculation on re-ordered instructions. Also fixed potential external function register alignment bugs.  Turned new VectorVM back on by default and bumped the Niagara script compiled version. #preflight 62438337c51908059aa24096 #rb rob.krajcarski

diff --git a/Engine/Plugins/FX/Niagara/Source/Niagara/Niagara.Build.cs b/Engine/Plugins/FX/Niagara/Source/Niagara/Niagara.Build.cs
--- a/Engine/Plugins/FX/Niagara/Source/Niagara/Niagara.Build.cs
+++ b/Engine/Plugins/FX/Niagara/Source/Niagara/Niagara.Build.cs
@@ -70,8 +70,8 @@ public class Niagara : ModuleRules
         PublicDefinitions.AddRange(
             new string[]
             {
-                "VECTORVM_SUPPORTS_EXPERIMENTAL=0",
-                "VECTORVM_SUPPORTS_LEGACY=1"
+                "VECTORVM_SUPPORTS_EXPERIMENTAL=1",
+                "VECTORVM_SUPPORTS_LEGACY=0"
             });
     }
 }
diff --git a/Engine/Source/Runtime/Core/Private/UObject/UE5MainStreamObjectVersion.cpp b/Engine/Source/Runtime/Core/Private/UObject/UE5MainStreamObjectVersion.cpp
--- a/Engine/Source/Runtime/Core/Private/UObject/UE5MainStreamObjectVersion.cpp
+++ b/Engine/Source/Runtime/Core/Private/UObject/UE5MainStreamObjectVersion.cpp
@@ -14,7 +14,7 @@ TMap<FGuid, FGuid> FUE5MainStreamObjectVersion::GetSystemGuids()
 	SystemGuids.Add(DevGuids.MATERIALSHADERMAP_DERIVEDDATA_VER, FGuid("65C01D817C9A4EEFAE9E988D41A1F3DD"));
 	SystemGuids.Add(DevGuids.NANITE_DERIVEDDATA_VER, FGuid("7E26CA5B0F5F03C40A3FD7173D88B8FA"));
 	SystemGuids.Add(DevGuids.NIAGARASHADERMAP_DERIVEDDATA_VER, FGuid("8A37C45D24F2423CBE5F8F371DE33575"));
-	SystemGuids.Add(DevGuids.Niagara_LatestScriptCompileVersion, FGuid("834E4334C1D7416CAF061A4FEFE4FE36"));
+	SystemGuids.Add(DevGuids.Niagara_LatestScriptCompileVersion, FGuid("078A8ED467E24F94BC5DDB4BE2A4EB8E"));
 	SystemGuids.Add(DevGuids.SkeletalMeshDerivedDataVersion, FGuid("FB7B2FB546DE49B2A2F88B7F69763CEE"));
 	SystemGuids.Add(DevGuids.STATICMESH_DERIVEDDATA_VER, FGuid("98F7D79A5811013825E75BBCC41ED3E9"));
 
diff --git a/Engine/Source/Runtime/VectorVM/Private/VectorVMExperimental.inl b/Engine/Source/Runtime/VectorVM/Private/VectorVMExperimental.inl
--- a/Engine/Source/Runtime/VectorVM/Private/VectorVMExperimental.inl
+++ b/Engine/Source/Runtime/VectorVM/Private/VectorVMExperimental.inl
@@ -411,9 +411,10 @@ static uint8 *SetupBatchStatePtrs(FVectorVMState *VVMState, FVectorVMBatchState
 		size_t PtrBeforeExtFnDecodeReg = (size_t)BatchDataPtr;
 		BatchState->ChunkLocalData.ExtFnDecodedReg.RegData       = (FVecReg **)BatchDataPtr;  BatchDataPtr += sizeof(FVecReg *) * VVMState->MaxExtFnRegisters;
 		BatchState->ChunkLocalData.ExtFnDecodedReg.RegInc        = (uint32 *)BatchDataPtr;    BatchDataPtr += sizeof(uint32)    * VVMState->MaxExtFnRegisters;
+		BatchDataPtr = (uint8 *)VVM_ALIGN_16(BatchDataPtr);
 		BatchState->ChunkLocalData.ExtFnDecodedReg.DummyRegs     = (FVecReg *)BatchDataPtr;   BatchDataPtr += sizeof(FVecReg)   * VVMState->NumDummyRegsReq;
 		size_t PtrAfterExtFnDecodeReg = (size_t)BatchDataPtr;
-		check(PtrAfterExtFnDecodeReg - PtrBeforeExtFnDecodeReg == VVMState->PerBatchChunkLocalNumExtFnDecodeRegisterBytesRequired);
+		check(PtrAfterExtFnDecodeReg - PtrBeforeExtFnDecodeReg <= VVMState->PerBatchChunkLocalNumExtFnDecodeRegisterBytesRequired);
 	}
 	
 	{ //after everything, likely outside of what the chunk will cache, setup the almost-never-used random counters and
@@ -535,8 +536,8 @@ VECTORVM_API FVectorVMState *InitVectorVMState(FVectorVMInitData *InitData, FVec
 	const uint32 MaxExtFnRegisters                                     = InitData->OptimizeContext->MaxExtFnRegisters == 0 ? 0 : VVM_ALIGN_4(InitData->OptimizeContext->MaxExtFnRegisters + 3); //we decode 4 at a time, so if we need any, we need a multiuple of 4
 	const size_t PerBatchChunkLocalDataOutputIdxBytesRequired	       = sizeof(uint32)                                        * InitData->OptimizeContext->NumOutputDataSets;                  //chunk local bytes required for instance offset
 	const size_t PerBatchChunkLocalNumOutputBytesRequired		       = sizeof(uint32)                                        * InitData->OptimizeContext->NumOutputDataSets;                  //chunk local bytes for num outputs
-	const size_t ConstantBufferSize					                   = sizeof(FVecReg)                                       * InitData->OptimizeContext->NumConstsRemapped;
-	const size_t PerBatchChunkLocalNumExtFnDecodeRegisterBytesRequired = (sizeof(FVecReg *) + sizeof(uint32)) * MaxExtFnRegisters + sizeof(FVecReg) * InitData->OptimizeContext->NumDummyRegsReq;
+	const size_t ConstantBufferSize					       = sizeof(FVecReg)                                       * InitData->OptimizeContext->NumConstsRemapped;
+	const size_t PerBatchChunkLocalNumExtFnDecodeRegisterBytesRequired = (sizeof(FVecReg *) + sizeof(uint32)) * MaxExtFnRegisters + sizeof(FVecReg) * (InitData->OptimizeContext->NumDummyRegsReq + 1); //+1 for potential padding
 
 	const size_t BatchOverheadSize                                     = ConstantBufferSize + 
 		                                                                 PerBatchChunkLocalDataOutputIdxBytesRequired + 
diff --git a/Engine/Source/Runtime/VectorVM/Private/VectorVMExperimental_Optimizer.inl b/Engine/Source/Runtime/VectorVM/Private/VectorVMExperimental_Optimizer.inl
--- a/Engine/Source/Runtime/VectorVM/Private/VectorVMExperimental_Optimizer.inl
+++ b/Engine/Source/Runtime/VectorVM/Private/VectorVMExperimental_Optimizer.inl
@@ -991,8 +991,6 @@ VECTORVM_API uint32 OptimizeVectorVMScript(const uint8 *InBytecode, int InByteco
 				{
 					int Idx = ExtFnIOData[ExtFnIdx].NumInputs + i;
 					check((VecIndices[Idx] & 0x8000) == 0 || VecIndices[Idx] == 0xFFFF); //can't output to a const... 0xFFFF is invalid
-					if (VecIndices[Idx] == 0xFFFF) {
-					}
 					OptContext->Intermediate.RegisterUsageBuffer[OptContext->Intermediate.NumRegistersUsed] = VecIndices[Idx];
 					++OptContext->Intermediate.NumRegistersUsed;
 				}
@@ -1106,10 +1104,12 @@ VECTORVM_API uint32 OptimizeVectorVMScript(const uint8 *InBytecode, int InByteco
 		}
 	}
 
-	int OnePastLastInputIdx = -1;
 	uint16 NumSSARegistersUsed = 0;
+	bool Step6_InstructionsRemovedRun = false;
 
 	{ //Step 5: SSA-like renaming of temp registers
+	Step5_SSA:
+		NumSSARegistersUsed = 0;
 		FVectorVMOptimizeInsRegUsage InputInsRegUse;
 		FVectorVMOptimizeInsRegUsage OutputInsRegUse;
 				
@@ -1206,12 +1206,72 @@ VECTORVM_API uint32 OptimizeVectorVMScript(const uint8 *InBytecode, int InByteco
 				DoneThisOutput: ;
 			}
 		}
-		check(SSARegCount < 0xFFFF - 1);
+		if (SSARegCount >= 0xFFFF) {
+			return VectorVMOptimizerSetError(OptContext, VVMOptErr_SSARemap | VVMOptErr_Overflow);
+		}
 		NumSSARegistersUsed = SSARegCount + 1;
 	}
 
-	
-	{ //Step 6: Input Fusing
+
+	if (1 && !Step6_InstructionsRemovedRun) { //Step 6: remove instructions where outputs are never used 
+		int NumRemovedInstructions = 0;
+		FVectorVMOptimizeInsRegUsage RegUsage;
+		FVectorVMOptimizeInsRegUsage RegUsage2;
+		int NumRemovedInstructionsThisTime;
+		int SanityCount = 0;
+		do
+		{
+			//loop multiple times because sometimes an instruction can be removed that will make a previous instruction redundant as well
+			NumRemovedInstructionsThisTime = 0;
+			for (uint32 i = 0; i < OptContext->Intermediate.NumInstructions; ++i)
+			{
+				FVectorVMOptimizeInstruction *Ins = OptContext->Intermediate.Instructions + i;
+				if (Ins->OpCat == EVectorVMOpCategory::Op && !(Ins->OpCode == EVectorVMOp::random || Ins->OpCode == EVectorVMOp::randomi))
+				{
+					//can we remove random instructions? I dunno! so lets not for now
+					GetRegistersUsedForInstruction(OptContext, Ins, &RegUsage);
+					for (int OutputIdx = 0; OutputIdx < RegUsage.NumOutputRegisters; ++OutputIdx)
+					{
+						uint16 RegIdx = OptContext->Intermediate.SSARegisterUsageBuffer[RegUsage.RegIndices[RegUsage.NumInputRegisters + OutputIdx]];
+						for (uint32 j = i + 1; j < OptContext->Intermediate.NumInstructions; ++j)
+						{
+							GetRegistersUsedForInstruction(OptContext, OptContext->Intermediate.Instructions + j, &RegUsage2);
+							for (int k = 0; k < RegUsage2.NumInputRegisters; ++k)
+							{
+								uint16 RegIdx2 = OptContext->Intermediate.SSARegisterUsageBuffer[RegUsage2.RegIndices[k]];
+								if (RegIdx == RegIdx2)
+								{
+									goto InstructionRequired;
+								}
+							}
+						}
+					}
+					{ //instruction isn't required
+						FMemory::Memmove(OptContext->Intermediate.Instructions + i, OptContext->Intermediate.Instructions + i + 1, sizeof(FVectorVMOptimizeInstruction) * (OptContext->Intermediate.NumInstructions - i - 1));
+						++NumRemovedInstructionsThisTime;
+						++NumRemovedInstructions;
+						--OptContext->Intermediate.NumInstructions;
+						--i;
+					}
+					InstructionRequired: ;
+				}
+			}
+			if (++SanityCount >= 16384)
+			{
+				return VectorVMOptimizerSetError(OptContext, VVMOptErr_RedundantInstruction);
+			}
+		} while (NumRemovedInstructionsThisTime > 0);
+
+		Step6_InstructionsRemovedRun = true;
+		if (NumRemovedInstructions > 0)
+		{
+			//if any instructions were removed we need to re-compute the SSA, so go and do that.
+			goto Step5_SSA;
+		}
+	}
+
+	int OnePastLastInputIdx = -1;
+	{ //Step 7: Input Fusing
 		//gather all input instructions can be fused
 		for (uint32 i = 0; i < OptContext->Intermediate.NumInstructions; ++i)
 		{
@@ -1229,7 +1289,11 @@ VECTORVM_API uint32 OptimizeVectorVMScript(const uint8 *InBytecode, int InByteco
 						case EVectorVMOpCategory::Input: //make sure there isn't an instruction that overwrites this input.  If this happens then this is a useless input
 							if (OptContext->Intermediate.SSARegisterUsageBuffer[InputIns->Input.DstRegPtrOffset] == OptContext->Intermediate.SSARegisterUsageBuffer[Ins_j->Input.DstRegPtrOffset])
 							{
-								check(false);
+								//check(false);
+								//This shouldn't happen.  I left the check() in for a while and it never fired.  I could probably remove this code, but maybe it'll help in debugging later
+								//I don't expect this check() will ever hit, but if something's going screwy with the input fusing maybe the compiler is generating 2 inputs that overwrite
+								//each other and that's causing an issue somewhere down the line, try uncommenting this check() and maybe it'll fire and give some insight... like I said
+								//though, I've never seen this check hit.
 							}
 							break;
 						case EVectorVMOpCategory::Output:
@@ -1237,7 +1301,7 @@ VECTORVM_API uint32 OptimizeVectorVMScript(const uint8 *InBytecode, int InByteco
 							{
 								if (Ins_j->OpCode != EVectorVMOp::outputdata_half) {
 									++InputIns->Input.FuseCount;
-									check(InputIns->Index == i);
+									check(InputIns->Index == i); //make sure the instruction is in its correct place.  Instructions could have moved, but this should have been corrected above.
 									Ins_j->Output.CopyFromInputInsIdx = i;
 								}
 							}
@@ -1328,7 +1392,9 @@ VECTORVM_API uint32 OptimizeVectorVMScript(const uint8 *InBytecode, int InByteco
 						case EVectorVMOpCategory::Stat:
 							break;
 						case EVectorVMOpCategory::Other:
-							check(Ins_j->OpCode == EVectorVMOp::done || Ins_j->OpCode == EVectorVMOp::noise2D || Ins_j->OpCode == EVectorVMOp::noise3D);
+							if (!(Ins_j->OpCode == EVectorVMOp::done || Ins_j->OpCode == EVectorVMOp::noise2D || Ins_j->OpCode == EVectorVMOp::noise3D)) {
+								return VectorVMOptimizerSetError(OptContext, VVMOptErr_Instructions);
+							}
 							break;
 					}
 					if (!InputOpCanFuse)
@@ -1348,136 +1414,7 @@ VECTORVM_API uint32 OptimizeVectorVMScript(const uint8 *InBytecode, int InByteco
 			++OnePastLastInputIdx;
 		}
 	}
-
-	if (0) { //Step 7: remove instructions where outputs are never used 
-		int NumRemovedInstructions = 0;
-		FVectorVMOptimizeInsRegUsage RegUsage;
-		FVectorVMOptimizeInsRegUsage RegUsage2;
-		int NumRemovedInstructionsThisTime;
-		int SanityCount = 0;
-		do
-		{
-			//loop multiple times because sometimes an instruction can be removed that will make a previous instruction redundant as well
-			NumRemovedInstructionsThisTime = 0;
-			for (uint32 i = 0; i < OptContext->Intermediate.NumInstructions; ++i)
-			{
-				FVectorVMOptimizeInstruction *Ins = OptContext->Intermediate.Instructions + i;
-				if (Ins->OpCat == EVectorVMOpCategory::Op && !(Ins->OpCode == EVectorVMOp::random || Ins->OpCode == EVectorVMOp::randomi))
-				{
-					//can we remove random instructions? I dunno! so lets not for now
-					GetRegistersUsedForInstruction(OptContext, Ins, &RegUsage);
-					for (int OutputIdx = 0; OutputIdx < RegUsage.NumOutputRegisters; ++OutputIdx)
-					{
-						uint16 RegIdx = OptContext->Intermediate.SSARegisterUsageBuffer[RegUsage.RegIndices[RegUsage.NumInputRegisters + OutputIdx]];
-						for (uint32 j = i + 1; j < OptContext->Intermediate.NumInstructions; ++j)
-						{
-							GetRegistersUsedForInstruction(OptContext, OptContext->Intermediate.Instructions + j, &RegUsage2);
-							for (int k = 0; k < RegUsage2.NumInputRegisters; ++k)
-							{
-								uint16 RegIdx2 = OptContext->Intermediate.SSARegisterUsageBuffer[RegUsage2.RegIndices[k]];
-								if (RegIdx == RegIdx2)
-								{
-									goto InstructionRequired;
-								}
-							}
-						}
-					}
-					{ //instruction isn't required
-						FMemory::Memmove(OptContext->Intermediate.Instructions + i, OptContext->Intermediate.Instructions + i + 1, sizeof(FVectorVMOptimizeInstruction) * (OptContext->Intermediate.NumInstructions - i - 1));
-						++NumRemovedInstructionsThisTime;
-						++NumRemovedInstructions;
-						--OptContext->Intermediate.NumInstructions;
-						--i;
-					}
-					InstructionRequired: ;
-				}
-			}
-			if (++SanityCount >= 16384)
-			{
-				check(false);
-				NumRemovedInstructions = 0;
-				break;
-			}
-		} while (NumRemovedInstructionsThisTime > 0);
-		
-		if (NumRemovedInstructions > 0) //Step 8: re-assign SSA registers if we removed instructions
-		{
-			FMemory::Memcpy(OptContext->Intermediate.SSARegisterUsageBuffer, OptContext->Intermediate.RegisterUsageBuffer, sizeof(uint16) * OptContext->Intermediate.NumRegistersUsed);
-			NumSSARegistersUsed = 0;
-
-			FVectorVMOptimizeInsRegUsage InputInsRegUse;
-			FVectorVMOptimizeInsRegUsage OutputInsRegUse;
-				
-			int SSARegCount = 0;
-			for (uint32 i = 0; i < OptContext->Intermediate.NumInstructions; ++i)
-			{
-				FVectorVMOptimizeInstruction *OutputIns = OptContext->Intermediate.Instructions + i;
-				GetRegistersUsedForInstruction(OptContext, OutputIns, &OutputInsRegUse);
-				if (OutputInsRegUse.NumOutputRegisters > 0)
-				{
-					for (int j = 0; j < OutputInsRegUse.NumOutputRegisters; ++j)
-					{
-						uint16 OutReg = OptContext->Intermediate.RegisterUsageBuffer[OutputInsRegUse.RegIndices[OutputInsRegUse.NumInputRegisters + j]];
-						if (OutReg != 0xFFFF) {
-							OptContext->Intermediate.SSARegisterUsageBuffer[OutputInsRegUse.RegIndices[OutputInsRegUse.NumInputRegisters + j]] = SSARegCount;
-							int LastUsedAsInputInsIdx = -1;
-							for (uint32 k = i + 1; k < OptContext->Intermediate.NumInstructions; ++k)
-							{
-								FVectorVMOptimizeInstruction *InputIns = OptContext->Intermediate.Instructions + k;
-								GetRegistersUsedForInstruction(OptContext, InputIns, &InputInsRegUse);
-								for (int ii = 0; ii < InputInsRegUse.NumOutputRegisters; ++ii)
-								{
-									if (OptContext->Intermediate.RegisterUsageBuffer[InputInsRegUse.RegIndices[InputInsRegUse.NumInputRegisters + ii]] == OutReg)
-									{
-										//this register is overwritten, we need to generate a new register
-										++SSARegCount;
-										check(SSARegCount <= (int)OptContext->Intermediate.NumRegistersUsed);
-										goto DoneThisOutput2;
-									}
-								}
-								for (int ii = 0; ii < InputInsRegUse.NumInputRegisters; ++ii)
-								{
-									if (OptContext->Intermediate.RegisterUsageBuffer[InputInsRegUse.RegIndices[ii]] == OutReg)
-									{
-										if (InputIns->OpCat == EVectorVMOpCategory::Output)
-										{
-											if (InputIns->OpCode != EVectorVMOp::copy_to_output)
-											{
-												if (OutputIns->OpCode == EVectorVMOp::acquireindex)
-												{
-													OptContext->Intermediate.SSARegisterUsageBuffer[InputIns->Output.RegPtrOffset] = SSARegCount;
-												}
-												else
-												{
-													OptContext->Intermediate.SSARegisterUsageBuffer[InputIns->Output.RegPtrOffset + 1] = SSARegCount;
-													LastUsedAsInputInsIdx = k;
-												}
-											}
-										}
-										else
-										{
-											LastUsedAsInputInsIdx = k;
-											OptContext->Intermediate.SSARegisterUsageBuffer[InputInsRegUse.RegIndices[ii]] = SSARegCount;
-										}
-									}
-								}
-							}
-							if (LastUsedAsInputInsIdx != -1)
-							{
-								++SSARegCount;
-							}
-						} else {
-							OptContext->Intermediate.SSARegisterUsageBuffer[OutputInsRegUse.RegIndices[OutputInsRegUse.NumInputRegisters + j]] = 0xFFFF;
-						}
-						DoneThisOutput2: ;
-					}
-				}
-			}
-			check(SSARegCount < 0xFFFF - 1);
-			NumSSARegistersUsed = SSARegCount + 1;
-		}
-	}
-
+	 
 	if (1) { //instruction re-ordering
 		int *RegToCheckStack = (int *)OptContext->Init.ReallocFn(nullptr, sizeof(int) * OptContext->Intermediate.NumRegistersUsed * 2, __FILE__, __LINE__);	//these two could actually be a single array, 1/2 the size, one starting from 0 and counting up, the other one
 		if (RegToCheckStack == nullptr)
@@ -1618,8 +1555,9 @@ VECTORVM_API uint32 OptimizeVectorVMScript(const uint8 *InBytecode, int InByteco
 				}
 			}
 		}
-
-		if (1) { //Step 12: re-order all dependent-less instructions to right before their output is used
+	
+		if (1)
+		{ //Step 12: re-order all dependent-less instructions to right before their output is used
 			int LastSwapInstructionIdx = -1; //to prevent an infinite loop when one instruction has two or more dependencies and they keep swapping back and forth
 			for (uint32 i = 0; i < OptContext->Intermediate.NumInstructions; ++i)
 			{
@@ -1661,8 +1599,8 @@ VECTORVM_API uint32 OptimizeVectorVMScript(const uint8 *InBytecode, int InByteco
 				}
 			}
 		}
-
-		if (1) { //Step 13: re-order all inputs to directly before they're used
+		if (1)
+		{ //Step 13: re-order all inputs to directly before they're used
 			FVectorVMOptimizeInsRegUsage RegUsage;
 			for (uint32 i = 0; i < OptContext->Intermediate.NumInstructions; ++i)
 			{
@@ -1719,7 +1657,6 @@ VECTORVM_API uint32 OptimizeVectorVMScript(const uint8 *InBytecode, int InByteco
 			}
 		}
 	}
-
 	{ //Step 14: if we re-ordered inputs, the CopyFromInputInsIdx on output instructions could be wrong.  Go fix those
 		for (uint32 OutputInsIdx = 0; OutputInsIdx < OptContext->Intermediate.NumInstructions; ++OutputInsIdx)
 		{
@@ -1946,7 +1883,9 @@ VECTORVM_API uint32 OptimizeVectorVMScript(const uint8 *InBytecode, int InByteco
 							break;
 						}
 					}
-					check(MinimizedRegIdx != 0xFFFF);
+					if (MinimizedRegIdx == 0xFFFF) {
+						return VectorVMOptimizerSetError(OptContext, VVMOptErr_RegisterUsage);
+					}
 				
 					OptContext->Intermediate.RegisterUsageBuffer[OutputRegIdx] = MinimizedRegIdx;
 
@@ -2225,6 +2164,7 @@ VECTORVM_API uint32 OptimizeVectorVMScript(const uint8 *InBytecode, int InByteco
 					{ //all inputs are regular registers, write the operation as normal
 						VVMOptWriteByte(Ins->OpCode);
 						for (int j = 0; j < Ins->Op.NumInputs + Ins->Op.NumOutputs; ++j) {
+							check(OptContext->Intermediate.RegisterUsageBuffer[Ins->Op.RegPtrOffset + j] != 0xFFFF);
 							VVMOptWriteU16(OptContext->Intermediate.RegisterUsageBuffer[Ins->Op.RegPtrOffset + j]);
 						}
 					}
@@ -2255,6 +2195,7 @@ VECTORVM_API uint32 OptimizeVectorVMScript(const uint8 *InBytecode, int InByteco
 								}
 								else
 								{
+									check(OptContext->Intermediate.RegisterUsageBuffer[Ins->Op.RegPtrOffset + j] != 0xFFFF);
 									VVMOptWriteU16(OptContext->Intermediate.RegisterUsageBuffer[Ins->Op.RegPtrOffset + j]);
 								}
 							}
diff --git a/Engine/Source/Runtime/VectorVM/Public/VectorVMExperimental.h b/Engine/Source/Runtime/VectorVM/Public/VectorVMExperimental.h
--- a/Engine/Source/Runtime/VectorVM/Public/VectorVMExperimental.h
+++ b/Engine/Source/Runtime/VectorVM/Public/VectorVMExperimental.h
@@ -66,8 +66,8 @@ struct FVectorVMSerializeDataSet
 	uint32 * InputBuffers;
 	uint32 * OutputBuffers;
 
-	uint32  InputOffset[3];	    //float, int, half (half must be 0)
-	uint32  OutputOffset[3];    //float, int, half (half must be 0)
+	uint32  InputOffset[4];	    //float, int, half (half must be 0)
+	uint32  OutputOffset[4];    //float, int, half (half must be 0)
 
 	int32   InputInstanceOffset;
 	int32   InputDataSetAccessIndex;
@@ -151,7 +151,7 @@ struct FVectorVMSerializeState
 	uint32                              NumChunks;
 	FVectorVMSerializeChunk *           Chunks;
 
-	struct FVectorVMOptimizeContext *   OptimizeCtx;
+	const struct FVectorVMOptimizeContext *   OptimizeCtx;
 
 	volatile int64                      ChunkComplete; //1 bit for each of the first 64 chunks
 
@@ -261,18 +261,20 @@ struct FVectorVMOptimizeInstruction
 
 enum EVectorVMOptimizeError
 {
-	VVMOptErr_OutOfMemory        = 1 << 0,
-	VVMOptErr_Bytecode           = 1 << 1,
-	VVMOptErr_RegisterUsage      = 1 << 2,
-	VVMOptErr_ConstRemap         = 1 << 3,
-	VVMOptErr_Instructions       = 1 << 4,
-	VVMOptErr_InputFuseBuffer    = 1 << 5,
-	VVMOptErr_InstructionReOrder = 1 << 6,
-	VVMOptErr_SSARemap           = 1 << 7,
-	VVMOptErr_OptimizedBytecode  = 1 << 8,
-	VVMOptErr_ExternalFunction   = 1 << 9,
-
-	VVMOptErr_Fatal              = 1 << 31
+	VVMOptErr_OutOfMemory           = 1 << 0,
+	VVMOptErr_Overflow			    = 1 << 1,
+	VVMOptErr_Bytecode              = 1 << 2,
+	VVMOptErr_RegisterUsage         = 1 << 3,
+	VVMOptErr_ConstRemap            = 1 << 4,
+	VVMOptErr_Instructions          = 1 << 5,
+	VVMOptErr_InputFuseBuffer       = 1 << 6,
+	VVMOptErr_InstructionReOrder    = 1 << 7,
+	VVMOptErr_SSARemap              = 1 << 8,
+	VVMOptErr_OptimizedBytecode     = 1 << 9,
+	VVMOptErr_ExternalFunction      = 1 << 10,
+	VVMOptErr_RedundantInstruction  = 1 << 11,
+
+	VVMOptErr_Fatal                 = 1 << 31
 };
 
 struct FVectorVMOptimizeContext
@@ -337,7 +339,7 @@ struct FVectorVMExternalFnPerInstanceData
 struct FVectorVMInitData
 {
 	struct FVectorVMState *					ExistingVectorVMState;
-	const FVectorVMOptimizeContext *              OptimizeContext;
+	const FVectorVMOptimizeContext *        OptimizeContext;
 	TArrayView<FDataSetMeta>                DataSets;
 	TArrayView<const FVMExternalFunction *> ExtFunctionTable;
 	
@@ -428,7 +430,7 @@ struct FVectorVMState
 class FVectorVMExternalFunctionContextExperimental
 {
 public:
-	uint32** RegisterData;
+	MS_ALIGN(16) uint32** RegisterData GCC_ALIGN(16);
 	uint16* RawVecIndices; //undecoded, for compatbility with the previous VM
 	uint32* RegInc;
 
diff --git a/Engine/Source/Runtime/VectorVM/VectorVM.Build.cs b/Engine/Source/Runtime/VectorVM/VectorVM.Build.cs
--- a/Engine/Source/Runtime/VectorVM/VectorVM.Build.cs
+++ b/Engine/Source/Runtime/VectorVM/VectorVM.Build.cs
@@ -27,8 +27,8 @@ public class VectorVM : ModuleRules
         PublicDefinitions.AddRange(
             new string[]
             {
-                "VECTORVM_SUPPORTS_EXPERIMENTAL=0",
-                "VECTORVM_SUPPORTS_LEGACY=1"
+                "VECTORVM_SUPPORTS_EXPERIMENTAL=1",
+                "VECTORVM_SUPPORTS_LEGACY=0",
             });
     }
 }
