Fix safe instructions for experimental VM to match the behavior of the legacy VM

diff --git a/Engine/Source/Runtime/VectorVM/Private/VectorVMExperimental.inl b/Engine/Source/Runtime/VectorVM/Private/VectorVMExperimental.inl
--- a/Engine/Source/Runtime/VectorVM/Private/VectorVMExperimental.inl
+++ b/Engine/Source/Runtime/VectorVM/Private/VectorVMExperimental.inl
@@ -329,10 +329,10 @@ VVMSet_m128iConst(  AlmostTwoBits   , 0x3fffffff);
 //safe instructions -- handle divide by zero "gracefully" by returning 0
 #define VVM_safeIns_div(v0, v1)      VectorSelect(VectorCompareGT(VectorAbs(v1), VVM_m128Const(Epsilon)), VectorDivide(v0, v1)                  , VectorZeroFloat())
 #define VVM_safeIns_rcp(v)           VectorSelect(VectorCompareGT(VectorAbs(v) , VVM_m128Const(Epsilon)), VectorDivide(VVM_m128Const(One), v)   , VectorZeroFloat())
-#define VVM_safe_sqrt(v)             VectorSelect(VectorCompareGT(VectorAbs(v) , VVM_m128Const(Epsilon)), VectorSqrt(v)                         , VectorZeroFloat())
-#define VVM_safe_log(v)              VectorSelect(VectorCompareGT(VectorAbs(v) , VectorZeroFloat())     , VectorLog(v)		                    , VectorZeroFloat())
-#define VVM_safe_pow(v0, v1)         VectorSelect(VectorCompareGT(VectorAbs(v1), VVM_m128Const(Epsilon)), VectorPow(v0, v1)                     , VectorZeroFloat())
-#define VVM_safe_rsq(v)              VectorSelect(VectorCompareGT(VectorAbs(v) , VVM_m128Const(Epsilon)), VectorReciprocalSqrt(v)               , VectorZeroFloat())
+#define VVM_safeins_sqrt(v)          VectorSelect(VectorCompareGT(v            , VVM_m128Const(Epsilon)), VectorSqrt(v)                         , VectorZeroFloat())
+#define VVM_safeins_log(v)           VectorSelect(VectorCompareGT(v            , VectorZeroFloat())     , VectorLog(v)		                    , VectorZeroFloat())
+#define VVM_safeins_pow(v0, v1)      VectorSelect(VectorCompareGT(v0           , VVM_m128Const(Epsilon)), VectorPow(v0, v1)                     , VectorZeroFloat())
+#define VVM_safeins_rsq(v)           VectorSelect(VectorCompareGT(v            , VVM_m128Const(Epsilon)), VectorReciprocalSqrt(v)               , VectorZeroFloat())
 
 #define VVMDebugBreakIf(expr)	if ((expr)) { PLATFORM_BREAK(); }
 
@@ -930,13 +930,13 @@ OpCodeSwitch: //I think computed gotos would be a huge win here... maybe write t
 				case EVectorVMOp::mad:                              execVecIns3f(VectorMultiplyAdd);    break;
 				case EVectorVMOp::lerp:                             execVecIns3f(VectorLerp);           break;
 				case EVectorVMOp::rcp:                              execVecIns1f(VVM_safeIns_rcp);      break;
-				case EVectorVMOp::rsq:                              execVecIns1f(VVM_safe_rsq);         break;
-				case EVectorVMOp::sqrt:                             execVecIns1f(VVM_safe_sqrt);        break;
+				case EVectorVMOp::rsq:                              execVecIns1f(VVM_safeins_rsq);      break;
+				case EVectorVMOp::sqrt:                             execVecIns1f(VVM_safeins_sqrt);     break;
 				case EVectorVMOp::neg:                              execVecIns1f(VectorNegate);         break;
 				case EVectorVMOp::abs:                              execVecIns1f(VectorAbs);            break;
 				case EVectorVMOp::exp:                              execVecIns1f(VectorExp);            break;
 				case EVectorVMOp::exp2:                             execVecIns1f(VectorExp2);           break;
-				case EVectorVMOp::log:                              execVecIns1f(VVM_safe_log);         break;
+				case EVectorVMOp::log:                              execVecIns1f(VVM_safeins_log);      break;
 				case EVectorVMOp::log2:                             execVecIns1f(VectorLog2);           break;
 				case EVectorVMOp::sin:                              execVecIns1f(VectorSin);            break;
 				case EVectorVMOp::cos:                              execVecIns1f(VectorCos);            break;
@@ -953,7 +953,7 @@ OpCodeSwitch: //I think computed gotos would be a huge win here... maybe write t
 				case EVectorVMOp::clamp:                            execVecIns3f(VectorClamp);          break;
 				case EVectorVMOp::min:                              execVecIns2f(VectorMin);            break;
 				case EVectorVMOp::max:                              execVecIns2f(VectorMax);            break;
-				case EVectorVMOp::pow:                              execVecIns2f(VVM_safe_pow);         break;
+				case EVectorVMOp::pow:                              execVecIns2f(VVM_safeins_pow);      break;
 				case EVectorVMOp::round:                            execVecIns1f(VectorRound);          break;
 				case EVectorVMOp::sign:                             execVecIns1f(VectorSign);           break;
 				case EVectorVMOp::step:                             execVecIns2f(VVM_vecStep);          break;
