Fix out of bounds access in experimental VM optimizer when handling a stat op as well as fixes up the index of referenced input ops when instructions are being reordered

diff --git a/Engine/Source/Runtime/VectorVM/Private/VectorVMExperimental_Optimizer.inl b/Engine/Source/Runtime/VectorVM/Private/VectorVMExperimental_Optimizer.inl
--- a/Engine/Source/Runtime/VectorVM/Private/VectorVMExperimental_Optimizer.inl
+++ b/Engine/Source/Runtime/VectorVM/Private/VectorVMExperimental_Optimizer.inl
@@ -1604,52 +1604,59 @@ VECTORVM_API uint32 OptimizeVectorVMScript(const uint8 *InBytecode, int InByteco
 			for (uint32 i = 0; i < OptContext->Intermediate.NumInstructions; ++i)
 			{
 				FVectorVMOptimizeInstruction *InputIns = OptContext->Intermediate.Instructions + i;
-				uint16 InputReg = OptContext->Intermediate.SSARegisterUsageBuffer[InputIns->Input.DstRegPtrOffset];
-				if (InputIns->OpCat == EVectorVMOpCategory::Input && InputIns->Input.FirstInsInsertIdx != -1)
+				if (InputIns->OpCat == EVectorVMOpCategory::Input)
 				{
-					for (uint32 j = i + 1; j < OptContext->Intermediate.NumInstructions; ++j)
+					uint16 InputReg = OptContext->Intermediate.SSARegisterUsageBuffer[InputIns->Input.DstRegPtrOffset];
+					if (InputIns->Input.FirstInsInsertIdx != -1)
 					{
-						FVectorVMOptimizeInstruction *OpIns = OptContext->Intermediate.Instructions + j;
-						if (OpIns->OpCat == EVectorVMOpCategory::Output && OpIns->Output.CopyFromInputInsIdx == i)
-						{
-							continue;
-						}
-						GetRegistersUsedForInstruction(OptContext, OpIns, &RegUsage);
-						bool MoveInputHere = false;
-						for (int k = 0; k < RegUsage.NumInputRegisters; ++k)
+						for (uint32 j = i + 1; j < OptContext->Intermediate.NumInstructions; ++j)
 						{
-							if (OptContext->Intermediate.SSARegisterUsageBuffer[RegUsage.RegIndices[k]] == InputReg)
+							bool MoveInputHere = false;
+
+							FVectorVMOptimizeInstruction *OpIns = OptContext->Intermediate.Instructions + j;
+							if (OpIns->OpCat == EVectorVMOpCategory::Output && OpIns->Output.CopyFromInputInsIdx == i)
 							{
 								MoveInputHere = true;
-								break;
 							}
-						}
-
-						if (MoveInputHere)
-						{
-							if (j > i + 1)
+							else
 							{
-								int NewInputIndex = j - 1;
-								FVectorVMOptimizeInstruction TempIns = *InputIns;
-								FMemory::Memmove(InputIns, InputIns + 1, sizeof(FVectorVMOptimizeInstruction) * (j - i - 1));
-								OptContext->Intermediate.Instructions[NewInputIndex] = TempIns;
-								bool ReorderingInputs = true;
-								//if we're only moving this instruction before inputs then we'll get into an infinite loop just re-ordering
-								//inputs around each other.  Check for that and skip all these if that's the case
-								check(i < j);
-								for (uint32 k = i; k < j; ++k)
+								GetRegistersUsedForInstruction(OptContext, OpIns, &RegUsage);
+								for (int k = 0; k < RegUsage.NumInputRegisters; ++k)
 								{
-									if (OptContext->Intermediate.Instructions[k].OpCat != EVectorVMOpCategory::Input)
+									if (OptContext->Intermediate.SSARegisterUsageBuffer[RegUsage.RegIndices[k]] == InputReg)
 									{
-										ReorderingInputs = false;
+										MoveInputHere = true;
+										break;
 									}
 								}
-								if (!ReorderingInputs)
+							}
+
+							if (MoveInputHere)
+							{
+								if (j > i + 1)
 								{
-									--i;
+									int NewInputIndex = j - 1;
+									FVectorVMOptimizeInstruction TempIns = *InputIns;
+									FMemory::Memmove(InputIns, InputIns + 1, sizeof(FVectorVMOptimizeInstruction) * (j - i - 1));
+									OptContext->Intermediate.Instructions[NewInputIndex] = TempIns;
+									bool ReorderingInputs = true;
+									//if we're only moving this instruction before inputs then we'll get into an infinite loop just re-ordering
+									//inputs around each other.  Check for that and skip all these if that's the case
+									check(i < j);
+									for (uint32 k = i; k < j; ++k)
+									{
+										if (OptContext->Intermediate.Instructions[k].OpCat != EVectorVMOpCategory::Input)
+										{
+											ReorderingInputs = false;
+										}
+									}
+									if (!ReorderingInputs)
+									{
+										--i;
+									}
 								}
+								break;
 							}
-							break;
 						}
 					}
 				}
@@ -1711,25 +1718,42 @@ VECTORVM_API uint32 OptimizeVectorVMScript(const uint8 *InBytecode, int InByteco
 					FVectorVMOptimizeInstruction TempIns = *Ins;
 					FMemory::Memmove(OptContext->Intermediate.Instructions + FirstCopyFromInputInsIdx + 1, OptContext->Intermediate.Instructions + FirstCopyFromInputInsIdx, sizeof(FVectorVMOptimizeInstruction) * (i - FirstCopyFromInputInsIdx));
 					OptContext->Intermediate.Instructions[FirstCopyFromInputInsIdx] = TempIns;
+
+					// if the moved instruction was an input we need to find all subsequent output ops and correct
+					// their CopyFromInputInsIdx (as necessary)
+					if (TempIns.OpCat == EVectorVMOpCategory::Input)
+					{
+						for (uint32 j = FirstCopyFromInputInsIdx + 1; j < OptContext->Intermediate.NumInstructions; ++j)
+						{
+							FVectorVMOptimizeInstruction* DownstreamIns = OptContext->Intermediate.Instructions + j;
+							if (DownstreamIns->OpCat == EVectorVMOpCategory::Output && DownstreamIns->Output.CopyFromInputInsIdx == i)
+							{
+								DownstreamIns->Output.CopyFromInputInsIdx = FirstCopyFromInputInsIdx;
+							}
+						}
+					}
+
 					++FirstCopyFromInputInsIdx;
 				}
 			}
 		}
 
-		
-		if (LastCopyFromInputInsIdx >= FirstCopyFromInputInsIdx + 2)
+		const int32 NumCopyInstructions = LastCopyFromInputInsIdx - FirstCopyFromInputInsIdx + 1;
+		if (NumCopyInstructions > 1)
 		{  // small list, very stupid bubble sort
 			bool sorted = false;
-			int NumCopyInstructions = LastCopyFromInputInsIdx - FirstCopyFromInputInsIdx + 1;
-			while (!sorted) {
+			while (!sorted)
+			{
 				sorted = true;
-				for (int i = 0; i < NumCopyInstructions - 1; ++i) {
+				for (int i = 0; i < NumCopyInstructions - 1; ++i)
+				{
 					FVectorVMOptimizeInstruction *Ins0 = OptContext->Intermediate.Instructions + FirstCopyFromInputInsIdx + i;
 					FVectorVMOptimizeInstruction *Ins1 = Ins0 + 1;
 					uint64 Key0 = VVMCopyToOutputInsGetSortKey(OptContext->Intermediate.Instructions, Ins0);
 					uint64 Key1 = VVMCopyToOutputInsGetSortKey(OptContext->Intermediate.Instructions, Ins1);
 
-					if (Key1 < Key0) {
+					if (Key1 < Key0)
+					{
 						FVectorVMOptimizeInstruction Temp = *Ins0;
 						*Ins0 = *Ins1;
 						*Ins1 = Temp;
