VectorVM Optimizations

diff --git a/Engine/Source/Runtime/VectorVM/Private/VectorVMExperimental.inl b/Engine/Source/Runtime/VectorVM/Private/VectorVMExperimental.inl
--- a/Engine/Source/Runtime/VectorVM/Private/VectorVMExperimental.inl
+++ b/Engine/Source/Runtime/VectorVM/Private/VectorVMExperimental.inl
@@ -24,6 +24,19 @@ pathological case:
 
 */
 
+/*
+	- OptimizeContext:
+		- static data cooked holds new bytecode, const and input remap data from the original bytecode to the new one.
+	- VectorVMState
+		- states to execute the VM that's independent of the number of instances required... mallocs memory once for the entire duration
+		  of the execution of a script... can/will have more than one per OptCtx.  Caches the input and const remap table upon the first
+		  exec() call
+	- ExecContext
+		- per-execution... can have more than one per VVMState
+	- Batch State
+		- per-batch.   Gets memory the batch memory cache... can have more than one per ExecCtx
+*/
+
 /*
 The three steps to running the new VM are:
 1. call OptimizeVectorVMScript() using the original bytecode and function bindings as input.
@@ -265,14 +278,6 @@ It is guaranteed by the optimizer that the index for output_batch8 and output_ba
 TempReg, not a ConstBuff so the decoding can be optimized.
 */
 
-enum EVVMRegFlags
-{
-	VVMRegFlag_Int      = 1,
-	VVMRegFlag_Clean    = 32,
-	VVMRegFlag_Index    = 64,
-	VVMRegFlag_Mismatch = 128,
-};
-
 static void *VVMDefaultRealloc(void *Ptr, size_t NumBytes, const char *Filename, int LineNumber)
 {
 	return FMemory::Realloc(Ptr, NumBytes);
@@ -300,15 +305,42 @@ static void VVMDefaultFree(void *Ptr, const char *Filename, int LineNumber)
 #define VVM_ALIGN_64(num)           (((size_t)(num) + 63) & ~63)
 #define VVM_ALIGN_CACHELINE(num)	(((size_t)(num) + (VVM_CACHELINE_SIZE - 1)) & ~(VVM_CACHELINE_SIZE - 1))
 
+#if VECTORVM_SUPPORTS_AVX
+#include "ThirdParty/SSEMathFun/avx_mathfun.h"
+#define VVM_PTR_ALIGN    VVM_ALIGN_32
+#define VVM_REG_SIZE     sizeof(__m256)
+#else
+#define VVM_PTR_ALIGN    VVM_ALIGN_16
+#define VVM_REG_SIZE     sizeof(FVecReg)
+#endif
+
 //to avoid memset/memcpy when statically initializing sse variables
-#define VVMSet_m128Const(Name, V)                static const MS_ALIGN(16) float VVMConstVec4_##Name##4[4]  GCC_ALIGN(16) = { V, V, V, V }
+#define VVMSet_m128Const(Name, V)                static const MS_ALIGN(16) float VVMConstVec4_##Name##4[4]   GCC_ALIGN(16) = { V, V, V, V }
+#define VVMSet_m128Const4(Name, V0, V1, V2, V3)  static const MS_ALIGN(16) float VVMConstVec4_##Name##4[4]   GCC_ALIGN(16) = { V0, V1, V2, V3 }
 #define VVMSet_m128iConst(Name, V)               static const MS_ALIGN(16) uint32 VVMConstVec4_##Name##4i[4] GCC_ALIGN(16) = { V, V, V, V }
 #define VVMSet_m128iConst4(Name, V0, V1, V2, V3) static const MS_ALIGN(16) uint32 VVMConstVec4_##Name##4i[4] GCC_ALIGN(16) = { V0, V1, V2, V3 }	/* equiv to setr */
 
 #define VVM_m128Const(Name)  (*(VectorRegister4f *)&(VVMConstVec4_##Name##4))
 #define VVM_m128iConst(Name) (*(VectorRegister4i *)&(VVMConstVec4_##Name##4i))
 
+#if VECTORVM_SUPPORTS_AVX
+#define VVM_m256Const(Name)  (*(__m256 *)&(VVMConstVec8_##Name##8))
+#define VVM_m256iConst(Name) (*(__m256i *)&(VVMConstVec8_##Name##8i))
+#define VVMSet_m256Const(Name, V)                                static const MS_ALIGN(32) float VVMConstVec8_##Name##8[8]   GCC_ALIGN(32) = { V, V, V, V }
+#define VVMSet_m256Const8(Name, V0, V1, V2, V3, V4, V5, V6, V7)  static const MS_ALIGN(32) float VVMConstVec8_##Name##8[8]   GCC_ALIGN(32) = { V0, V1, V2, V3, V4, V5, V6, V7 }
+#define VVMSet_m256iConst(Name, V)                               static const MS_ALIGN(32) uint32 VVMConstVec8_##Name##8i[8] GCC_ALIGN(32) = { V, V, V, V }
+#define VVMSet_m256iConst8(Name, V0, V1, V2, V3, V4, V5, V6, V7) static const MS_ALIGN(32) uint32 VVMConstVec8_##Name##8i[8] GCC_ALIGN(32) = { V0, V1, V2, V3, V4, V5, V6, V7 }	/* equiv to setr */
+#else
+#define VVM_m256Const(Name)
+#define VVM_m256iConst(Name)
+#define VVMSet_m256Const(Name, V)
+#define VVMSet_m256Const8(Name, V0, V1, V2, V3, V4, V5, V6, V7)
+#define VVMSet_m256iConst(Name, V)
+#define VVMSet_m256iConst8(Name, V0, V1, V2, V3, V4, V5, V6, V7)
+#endif
+
 VVMSet_m128Const(   One             , 1.f);
+VVMSet_m128Const(   NegativeOne     , -1.f);
 VVMSet_m128Const(   OneHalf         , 0.5f);
 VVMSet_m128Const(   Epsilon         , 1.e-8f);
 VVMSet_m128Const(   HalfPi          , 3.14159265359f * 0.5f);
@@ -316,31 +348,477 @@ VVMSet_m128Const(   FastSinA        , 7.5894663844f);
 VVMSet_m128Const(   FastSinB        , 1.6338434578f);
 VVMSet_m128iConst(  FMask           , 0xFFFFFFFF);
 VVMSet_m128iConst4( ZeroOneTwoThree , 0, 1, 2, 3);
+VVMSet_m128iConst4( ZeroTwoFourSix  , 0, 2, 4, 6);
+VVMSet_m128Const4(  ZeroOneTwoThree , 0.f, 1.f, 2.f, 3.f);
 VVMSet_m128iConst(  RegOffsetMask   , 0x7FFF);
 VVMSet_m128Const(   RegOneOverTwoPi , 1.f / 2.f / 3.14159265359f);
 VVMSet_m128iConst(  AlmostTwoBits   , 0x3fffffff);
 
-#define VVM_vecStep(a, b)            VectorStep(VectorSubtract(a, b))
-#define VVM_vecFloatToBool(v)        VectorCompareGT(v, VectorZeroFloat())
-#define VVM_vecBoolToFloat(v)        VectorSelect(v, VectorSet1(1.f), VectorZeroFloat());
-#define VVM_vecIntToBool(v)          VectorIntCompareGT(v, VectorSetZero())
-#define VVM_vecBoolToInt(v)          VectorIntSelect(v, VectorIntSet1(1), VectorSetZero())
-#define VVM_vecACosFast(v)           VectorATan2(VectorSqrt(VectorMultiply(VectorSubtract(VVM_m128Const(One), v), VectorAdd(VVM_m128Const(One), v))), v)
+VVMSet_m256Const(   One                                  , 1.f);
+VVMSet_m256Const(   NegativeOne                          , -1.f);
+VVMSet_m256Const(   OneHalf                              , 0.5f);
+VVMSet_m256Const(   Epsilon                              , 1.e-8f);
+VVMSet_m256Const(   HalfPi                               , 3.14159265359f * 0.5f);
+VVMSet_m256Const(   FastSinA                             , 7.5894663844f) ;
+VVMSet_m256Const(   FastSinB                             , 1.6338434578f);
+VVMSet_m256iConst(  FMask                                , 0xFFFFFFFF);
+VVMSet_m256iConst(  One                                  , 1);
+VVMSet_m256iConst(  NegativeOne                          , 0xFFFFFFFF);
+VVMSet_m256iConst8( ZeroOneTwoThreeFourFiveSixSeven      , 0, 1, 2, 3, 4, 5, 6, 7);
+VVMSet_m256iConst8( ZeroTwoFourSixEightTenTwelveFourteen , 0, 2, 4, 6, 8, 10, 12, 14);
+VVMSet_m256Const8(  ZeroOneTwoThreeFourFiveSixSeven      , 0.f, 1.f, 2.f, 3.f, 4.f, 5.f, 6.f, 7.f);
+VVMSet_m256iConst(  RegOffsetMask                        , 0x7FFF);
+VVMSet_m256Const(   RegOneOverTwoPi                      , 1.f / 2.f / 3.14159265359f);
+VVMSet_m256iConst(  AlmostTwoBits                        , 0x3fffffff);
+
+#define VVM_vecStep(a, b)                  VectorStep(VectorSubtract(a, b))
+#define VVM_vecFloatToBool(v)              VectorCompareGT(v, VectorZeroFloat())
+#define VVM_vecBoolToFloat(v)              VectorSelect(v, VVM_m128Const(One), VectorZeroFloat());
+#define VVM_vecIntToBool(v)                VectorIntCompareGT(v, VectorSetZero())
+#define VVM_vecBoolToInt(v)                VectorIntSelect(v, VectorIntSet1(1), VectorSetZero())
+#define VVM_vecACosFast(v)                 VectorATan2(VectorSqrt(VectorMultiply(VectorSubtract(VVM_m128Const(One), v), VectorAdd(VVM_m128Const(One), v))), v)
+//safe instructions -- handle divide by zero "gracefully" by returning 0
+#define VVM_safeIns_div(v0, v1)            VectorSelect(VectorCompareGT(VectorAbs(v1), VVM_m128Const(Epsilon)), VectorDivide(v0, v1)                  , VectorZeroFloat())
+#define VVM_safeIns_rcp(v)                 VectorSelect(VectorCompareGT(VectorAbs(v) , VVM_m128Const(Epsilon)), VectorVMAccuracy::Reciprocal(v)       , VectorZeroFloat())
+#define VVM_safe_sqrt(v)                   VectorSelect(VectorCompareGT(v            , VVM_m128Const(Epsilon)), VectorVMAccuracy::Sqrt(v)             , VectorZeroFloat())
+#define VVM_safe_log(v)                    VectorSelect(VectorCompareGT(v            , VectorZeroFloat())     , VectorLog(v)		                  , VectorZeroFloat())
+#define VVM_safe_pow(v0, v1)               VectorSelect(VectorCompareGT(v1           , VVM_m128Const(Epsilon)), VectorPow(v0, v1)                     , VectorZeroFloat())
+#define VVM_safe_rsq(v)                    VectorSelect(VectorCompareGT(v            , VVM_m128Const(Epsilon)), VectorReciprocalSqrt(v)               , VectorZeroFloat())
+#define VVM_random(v)                      VectorMultiply(VectorSubtract(VectorRegister4f(VectorCastIntToFloat(VectorIntOr(VectorShiftRightImmLogical(VVMXorwowStep(BatchState), 9), VectorIntSet1(0x3F800000)))), VVM_m128Const(One)), v)
+#define VVM_randomi(v)                     VectorFloatToInt(VectorMultiply(VectorSubtract(VectorRegister4f(VectorCastIntToFloat(VectorIntOr(VectorShiftRightImmLogical(VVMXorwowStep(BatchState), 9), VectorIntSet1(0x3F800000)))), VVM_m128Const(One)), *(VectorRegister4f *)&v))
+
+//new merged instructions
+#define VVM_cmplt_select(v0, v1, v2, v3)    VectorSelect(VectorCompareLT(v0, v1), v2, v3)
+#define VVM_cmple_select(v0, v1, v2, v3)    VectorSelect(VectorCompareLE(v0, v1), v2, v3)
+#define VVM_cmpeq_select(v0, v1, v2, v3)    VectorSelect(VectorCompareEQ(v0, v1), v2, v3)
+#define VVM_cmplti_select(v0, v1, v2, v3)   VectorSelect(VectorCastIntToFloat(VectorIntCompareLT(*(VectorRegister4i *)&v0, *(VectorRegister4i *)&v1)), v2, v3)
+#define VVM_cmplei_select(v0, v1, v2, v3)   VectorSelect(VectorCastIntToFloat(VectorIntCompareLE(*(VectorRegister4i *)&v0, *(VectorRegister4i *)&v1)), v2, v3)
+#define VVM_cmpeqi_select(v0, v1, v2, v3)   VectorSelect(VectorCastIntToFloat(VectorIntCompareEQ(*(VectorRegister4i *)&v0, *(VectorRegister4i *)&v1)), v2, v3)
+#define VVM_cmplt_logic_and(v0, v1, v2)     VectorCastIntToFloat(VectorIntAnd(VectorCastFloatToInt(VectorCompareLT(v0, v1)), *(VectorRegister4i *)&v2))
+#define VVM_cmple_logic_and(v0, v1, v2)     VectorCastIntToFloat(VectorIntAnd(VectorCastFloatToInt(VectorCompareLE(v0, v1)), *(VectorRegister4i *)&v2))
+#define VVM_cmpgt_logic_and(v0, v1, v2)     VectorCastIntToFloat(VectorIntAnd(VectorCastFloatToInt(VectorCompareGT(v0, v1)), *(VectorRegister4i *)&v2))
+#define VVM_cmpge_logic_and(v0, v1, v2)     VectorCastIntToFloat(VectorIntAnd(VectorCastFloatToInt(VectorCompareGE(v0, v1)), *(VectorRegister4i *)&v2))
+#define VVM_cmpeq_logic_and(v0, v1, v2)     VectorCastIntToFloat(VectorIntAnd(VectorCastFloatToInt(VectorCompareEQ(v0, v1)), *(VectorRegister4i *)&v2))
+#define VVM_cmpne_logic_and(v0, v1, v2)     VectorCastIntToFloat(VectorIntAnd(VectorCastFloatToInt(VectorCompareNE(v0, v1)), *(VectorRegister4i *)&v2))
+#define VVM_cmplti_logic_and(v0, v1, v2)    VectorIntAnd(VectorIntCompareLT(v0, v1), v2)
+#define VVM_cmplei_logic_and(v0, v1, v2)    VectorIntAnd(VectorIntCompareLE(v0, v1), v2)
+#define VVM_cmpgti_logic_and(v0, v1, v2)    VectorIntAnd(VectorIntCompareGT(v0, v1), v2)
+#define VVM_cmpgei_logic_and(v0, v1, v2)    VectorIntAnd(VectorIntCompareGE(v0, v1), v2)
+#define VVM_cmpeqi_logic_and(v0, v1, v2)    VectorIntAnd(VectorIntCompareEQ(v0, v1), v2)
+#define VVM_cmpnei_logic_and(v0, v1, v2)    VectorIntAnd(VectorIntCompareNEQ(v0, v1), v2)
+#define VVM_cmplt_logic_or(v0, v1, v2)      VectorCastIntToFloat(VectorIntOr(VectorCastFloatToInt(VectorCompareLT(v0, v1)), *(VectorRegister4i *)&v2))
+#define VVM_cmple_logic_or(v0, v1, v2)      VectorCastIntToFloat(VectorIntOr(VectorCastFloatToInt(VectorCompareLE(v0, v1)), *(VectorRegister4i *)&v2))
+#define VVM_cmpgt_logic_or(v0, v1, v2)      VectorCastIntToFloat(VectorIntOr(VectorCastFloatToInt(VectorCompareGT(v0, v1)), *(VectorRegister4i *)&v2))
+#define VVM_cmpge_logic_or(v0, v1, v2)      VectorCastIntToFloat(VectorIntOr(VectorCastFloatToInt(VectorCompareGE(v0, v1)), *(VectorRegister4i *)&v2))
+#define VVM_cmpeq_logic_or(v0, v1, v2)      VectorCastIntToFloat(VectorIntOr(VectorCastFloatToInt(VectorCompareEQ(v0, v1)), *(VectorRegister4i *)&v2))
+#define VVM_cmpne_logic_or(v0, v1, v2)      VectorCastIntToFloat(VectorIntOr(VectorCastFloatToInt(VectorCompareNE(v0, v1)), *(VectorRegister4i *)&v2))
+#define VVM_cmplti_logic_or(v0, v1, v2)     VectorIntOr(VectorIntCompareLT(v0, v1), v2)
+#define VVM_cmplei_logic_or(v0, v1, v2)     VectorIntOr(VectorIntCompareLE(v0, v1), v2)
+#define VVM_cmpgti_logic_or(v0, v1, v2)     VectorIntOr(VectorIntCompareGT(v0, v1), v2)
+#define VVM_cmpgei_logic_or(v0, v1, v2)     VectorIntOr(VectorIntCompareGE(v0, v1), v2)
+#define VVM_cmpeqi_logic_or(v0, v1, v2)     VectorIntOr(VectorIntCompareEQ(v0, v1), v2)
+#define VVM_cmpnei_logic_or(v0, v1, v2)     VectorIntOr(VectorIntCompareNEQ(v0, v1), v2)
+#define VVM_mad_add(v0, v1, v2, v3)         VectorAdd(VectorMultiplyAdd(v0, v1, v2), v3)
+#define VVM_mad_sub0(v0, v1, v2, v3)        VectorSubtract(VectorMultiplyAdd(v0, v1, v2), v3)
+#define VVM_mad_sub1(v0, v1, v2, v3)        VectorSubtract(v3, VectorMultiplyAdd(v0, v1, v2))
+#define VVM_mad_mul(v0, v1, v2, v3)         VectorMultiply(VectorMultiplyAdd(v0, v1, v2), v3)
+#define VVM_mad_sqrt(v0, v1, v2)            VectorSqrt(VectorMultiplyAdd(v0, v1, v2))
+#define VVM_mad_mad0(v0, v1, v2, v3, v4)    VectorMultiplyAdd(v3, v4, VectorMultiplyAdd(v0, v1, v2))
+#define VVM_mad_mad1(v0, v1, v2, v3, v4)    VectorMultiplyAdd(VectorMultiplyAdd(v0, v1, v2), v3, v4)
+#define VVM_mul_mad0(v0, v1, v2, v3)        VectorMultiplyAdd(VectorMultiply(v0, v1), v2, v3)
+#define VVM_mul_mad1(v0, v1, v2, v3)        VectorMultiplyAdd(v2, v3, VectorMultiply(v0, v1))
+#define VVM_mul_add(v0, v1, v2)             VectorAdd(VectorMultiply(v0, v1), v2)
+#define VVM_mul_sub0(v0, v1, v2)            VectorSubtract(VectorMultiply(v0, v1), v2)
+#define VVM_mul_sub1(v0, v1, v2)            VectorSubtract(v2, VectorMultiply(v0, v1))
+#define VVM_mul_mul(v0, v1, v2)             VectorMultiply(VectorMultiply(v0, v1), v2)
+#define VVM_mul_max(v0, v1, v2)             VectorMax(VectorMultiply(v0, v1), v2)
+#define VVM_add_mad1(v0, v1, v2, v3)        VectorMultiplyAdd(v2, v3, VectorAdd(v0, v1))
+#define VVM_add_add(v0, v1, v2)             VectorAdd(VectorAdd(v0, v1), v2)
+#define VVM_sub_cmplt1(v0, v1, v2)			VectorCompareLT(v2, VectorSubtract(v0, v1))
+#define VVM_sub_neg(v0, v1)                 VectorNegate(VectorSubtract(v0, v1))
+#define VVM_sub_mul(v0, v1, v2)				VectorMultiply(VectorSubtract(v0, v1), v2)
+#define VVM_div_mad0(v0, v1, v2, v3)        VectorMultiplyAdd(VVM_safeIns_div(v0, v1), v2, v3)
+#define VVM_div_f2i(v0, v1)                 VectorFloatToInt(VVM_safeIns_div(VectorCastIntToFloat(v0), VectorCastIntToFloat(v1)))
+#define VVM_div_mul(v0, v1, v2)				VectorMultiply(VVM_safeIns_div(v0, v1), v2)
+#define VVM_muli_addi(v0, v1, v2)           VectorIntAdd(VectorIntMultiply(v0, v1), v2)
+#define VVM_addi_bit_rshift(v0, v1, v2)     VVMIntRShift(VectorIntAdd(v0, v1), v2)
+#define VVM_addi_muli(v0, v1, v2)           VectorIntMultiply(VectorIntAdd(v0, v1), v2)
+#define VVM_i2f_div0(v0, v1)                VVM_safeIns_div(VectorIntToFloat(VectorCastFloatToInt(v0)), v1)
+#define VVM_i2f_div1(v0, v1)                VVM_safeIns_div(v1, VectorIntToFloat(VectorCastFloatToInt(v0))) 
+#define VVM_i2f_mul(v0, v1)                 VectorMultiply(VectorIntToFloat(VectorCastFloatToInt(v0)), v1)
+#define VVM_i2f_mad0(v0, v1, v2)            VectorMultiplyAdd(VectorIntToFloat(VectorCastFloatToInt(v0)), v1, v2)
+#define VVM_i2f_mad1(v0, v1, v2)            VectorMultiplyAdd(v0, v1, VectorIntToFloat(VectorCastFloatToInt(v2)))
+#define VVM_f2i_select1(mask, v0, v1)       VectorIntSelect(mask, VectorFloatToInt(VectorCastIntToFloat(v0)), v1)
+#define VVM_f2i_maxi(v0, v1)                VectorIntMax(VectorFloatToInt(VectorCastIntToFloat(v0)), v1)
+#define VVM_f2i_addi(v0, v1)                VectorIntAdd(VectorFloatToInt(VectorCastIntToFloat(v0)), v1)
+#define VVM_fmod_add(v0, v1, v2)            VectorAdd(VectorMod(v0, v1), v2)
+#define VVM_bit_and_i2f(v0, v1)             VectorIntToFloat(VectorIntAnd(VectorCastFloatToInt(v0), VectorCastFloatToInt(v1)))
+#define VVM_bit_rshift_bit_and(v0, v1, v2)  VectorIntAnd(VVMIntRShift(v0, v1), v2)
+#define VVM_neg_cmplt(v0, v1)               VectorCompareLT(VectorNegate(v0), v1)
+#define VVM_bit_or_muli(v0, v1, v2)         VectorIntMultiply(VectorIntOr(v0, v1), v2)
+#define VVM_bit_lshift_bit_or(v0, v1, v2)   VectorIntOr(VVMIntLShift(v0, v1), v2)
+#define VVM_random_add(v0, v1)              VectorAdd(VVM_random(v0), v1)
+#define VVM_max_f2i(v0, v1)                 VectorFloatToInt(VectorMax(VectorCastIntToFloat(v0), VectorCastIntToFloat(v1)))
+#define VVM_select_mul(v0, v1, v2, v3)      VectorMultiply(VectorSelect(v0, v1, v2), v3)
+#define VVM_select_add(v0, v1, v2, v3)      VectorAdd(VectorSelect(v0, v1, v2), v3)
+
+
+#if VECTORVM_SUPPORTS_AVX
+#define VVM_select_AVX(m, v0, v1)			_mm256_blendv_ps(v1, v0, m)
+#define VVM_selecti_AVX(m, v0, v1)			_mm256_blendv_epi8(v1, v0, m)
+#define VVM_abs_AVX(v)                      _mm256_and_ps(v, _mm256_castsi256_ps(_mm256_set1_epi32(~(1 << 31))))
+#define VVM_pow_AVX(base, exp)               exp256_ps(_mm256_mul_ps(log256_ps(base), exp))
+#define VVM_neg_AVX(v)                      _mm256_sub_ps(_mm256_setzero_ps(), v)
+
+union VVMXXMYMMUnion
+{
+	__m256 ymm;
+	__m128 xmm[2];
+	__m256i ymmi;
+	__m128i xmmi[2];
+};
+
+//AVX compare macros
+#define VVM_cmplt_AVX(v0, v1)               _mm256_cmp_ps(v0, v1, _CMP_LT_OS)
+#define VVM_cmple_AVX(v0, v1)               _mm256_cmp_ps(v0, v1, _CMP_LE_OS)
+#define VVM_cmpgt_AVX(v0, v1)               _mm256_cmp_ps(v0, v1, _CMP_GT_OS)
+#define VVM_cmpge_AVX(v0, v1)               _mm256_cmp_ps(v0, v1, _CMP_GE_OS)
+#define VVM_cmpeq_AVX(v0, v1)               _mm256_cmp_ps(v0, v1, _CMP_EQ_OQ)
+#define VVM_cmpneq_AVX(v0, v1)              _mm256_cmp_ps(v0, v1, _CMP_NEQ_OQ)
+
+#define VVM_cmplti_AVX(v0, v1)              _mm256_cmpgt_epi32(v1, v0)
+#define VVM_cmplei_AVX(v0, v1)              _mm256_or_si256(_mm256_cmpgt_epi32(v1, v0), _mm256_cmpeq_epi32(v0, v1))
+#define VVM_cmpgti_AVX(v0, v1)              _mm256_cmpgt_epi32(v0, v1)
+#define VVM_cmpgei_AVX(v0, v1)              _mm256_or_si256(_mm256_cmpgt_epi32(v0, v1), _mm256_cmpeq_epi32(v0, v1))
+#define VVM_cmpeqi_AVX(v0, v1)              _mm256_cmpeq_epi32(v0, v1)
+#define VVM_cmpneqi_AVX(v0, v1)             _mm256_xor_si256(_mm256_cmpeq_epi32(v0, v1), VVM_m256iConst(FMask))
+
+#define VVM_vecStep_AVX(v0, v1)             VVM_select_AVX(VVM_cmpge_AVX(_mm256_sub_ps(v0, v1), _mm256_setzero_ps()), VVM_m256Const(One), _mm256_setzero_ps())
+#define VVM_vecFloatToBool_AVX(v)           VVM_cmpgt_AVX  (v, _mm256_setzero_ps())
+#define VVM_vecBoolToFloat_AVX(v)           VVM_select_AVX (v, VVM_m256Const(One), _mm256_setzero_ps());
+#define VVM_vecIntToBool_AVX(v)             VVM_cmpgti_AVX (v, _mm256_setzero_si256())
+#define VVM_vecBoolToInt_AVX(v)             VVM_selecti_AVX(v, _mm256_set1_epi32(1), _mm256_setzero_si256())
+#define VVM_vecACosFast_AVX(v)              VectorATan2_AVX(_mm256_sqrt_ps(_mm256_mul_ps(_mm256_sub_ps(VVM_m256Const(One), v), _mm256_add_ps(VVM_m256Const(One), v))), v)
+//VectorATan2(VectorSqrt(VectorMultiply(VectorSubtract(VVM_m128Const(One), v), VectorAdd(VVM_m128Const(One), v))), v)
+
 //safe instructions -- handle divide by zero "gracefully" by returning 0
-#define VVM_safeIns_div(v0, v1)      VectorSelect(VectorCompareGT(VectorAbs(v1), VVM_m128Const(Epsilon)), VectorDivide(v0, v1)                  , VectorZeroFloat())
-#define VVM_safeIns_rcp(v)           VectorSelect(VectorCompareGT(VectorAbs(v) , VVM_m128Const(Epsilon)), VectorDivide(VVM_m128Const(One), v)   , VectorZeroFloat())
-#define VVM_safeins_sqrt(v)          VectorSelect(VectorCompareGT(v            , VVM_m128Const(Epsilon)), VectorSqrt(v)                         , VectorZeroFloat())
-#define VVM_safeins_log(v)           VectorSelect(VectorCompareGT(v            , VectorZeroFloat())     , VectorLog(v)		                    , VectorZeroFloat())
-#define VVM_safeins_pow(v0, v1)      VectorSelect(VectorCompareGT(v0           , VVM_m128Const(Epsilon)), VectorPow(v0, v1)                     , VectorZeroFloat())
-#define VVM_safeins_rsq(v)           VectorSelect(VectorCompareGT(v            , VVM_m128Const(Epsilon)), VectorReciprocalSqrt(v)               , VectorZeroFloat())
+#define VVM_safeIns_div_AVX(v0, v1)        VVM_select_AVX(VVM_cmpgt_AVX(VVM_abs_AVX(v1), VVM_m256Const(Epsilon)), _mm256_div_ps(v0, v1)                 , _mm256_setzero_ps())
+#define VVM_safeIns_rcp_AVX(v)             VVM_select_AVX(VVM_cmpgt_AVX(VVM_abs_AVX(v) , VVM_m256Const(Epsilon)), _mm256_rcp_ps(v)                      , _mm256_setzero_ps())
+#define VVM_safe_sqrt_AVX(v)               VVM_select_AVX(VVM_cmpgt_AVX(v              , VVM_m256Const(Epsilon)), _mm256_sqrt_ps(v)                     , _mm256_setzero_ps())
+#define VVM_safe_log_AVX(v)                VVM_select_AVX(VVM_cmpgt_AVX(v              , _mm256_setzero_ps())   , log256_ps(v)		                    , _mm256_setzero_ps())
+#define VVM_safe_pow_AVX(v0, v1)           VVM_select_AVX(VVM_cmpgt_AVX(v1             , VVM_m256Const(Epsilon)), VVM_pow_AVX(v0, v1)                   , _mm256_setzero_ps())
+#define VVM_safe_rsq_AVX(v)                VVM_select_AVX(VVM_cmpgt_AVX(v              , VVM_m256Const(Epsilon)), _mm256_rsqrt_ps (v)                   , _mm256_setzero_ps())
+#define VVM_random_AVX(v)                  _mm256_mul_ps(_mm256_sub_ps(_mm256_castsi256_ps(_mm256_or_si256(_mm256_srli_epi32(VVMXorwowStep_AVX(BatchState), 9), _mm256_set1_epi32(0x3F800000))), VVM_m256Const(One)), v)
+#define VVM_randomi_AVX(v)                 _mm256_cvttps_epi32(_mm256_mul_ps(_mm256_sub_ps(_mm256_castsi256_ps(_mm256_or_si256(_mm256_srli_epi32(VVMXorwowStep_AVX(BatchState), 9), _mm256_set1_epi32(0x3F800000))), VVM_m256Const(One)), *(__m256 *)&v))
+
+//new merged instructions	
+#define VVM_cmplt_select_AVX(v0, v1, v2, v3)    VVM_select_AVX(VVM_cmplt_AVX(v0, v1), v2, v3)
+#define VVM_cmple_select_AVX(v0, v1, v2, v3)    VVM_select_AVX(VVM_cmple_AVX(v0, v1), v2, v3)
+#define VVM_cmpeq_select_AVX(v0, v1, v2, v3)    VVM_select_AVX(VVM_cmpeq_AVX(v0, v1), v2, v3)
+#define VVM_cmplti_select_AVX(v0, v1, v2, v3)   VVM_select_AVX(_mm256_castsi256_ps(VVM_cmplti_AVX(*(__m256i *)&v0, *(__m256i *)&v1)), v2, v3)
+#define VVM_cmplei_select_AVX(v0, v1, v2, v3)   VVM_select_AVX(_mm256_castsi256_ps(VVM_cmplei_AVX(*(__m256i *)&v0, *(__m256i *)&v1)), v2, v3)
+#define VVM_cmpeqi_select_AVX(v0, v1, v2, v3)   VVM_select_AVX(_mm256_castsi256_ps(VVM_cmpeqi_AVX(*(__m256i *)&v0, *(__m256i *)&v1)), v2, v3)
+#define VVM_cmplt_logic_and_AVX(v0, v1, v2)     _mm256_castsi256_ps(_mm256_and_si256(_mm256_castps_si256(VVM_cmplt_AVX(v0, v1)), *(__m256i *)&v2))
+#define VVM_cmple_logic_and_AVX(v0, v1, v2)     _mm256_castsi256_ps(_mm256_and_si256(_mm256_castps_si256(VVM_cmple_AVX(v0, v1)), *(__m256i *)&v2))
+#define VVM_cmpgt_logic_and_AVX(v0, v1, v2)     _mm256_castsi256_ps(_mm256_and_si256(_mm256_castps_si256(VVM_cmpgt_AVX(v0, v1)), *(__m256i *)&v2))
+#define VVM_cmpge_logic_and_AVX(v0, v1, v2)     _mm256_castsi256_ps(_mm256_and_si256(_mm256_castps_si256(VVM_cmpge_AVX(v0, v1)), *(__m256i *)&v2))
+#define VVM_cmpeq_logic_and_AVX(v0, v1, v2)     _mm256_castsi256_ps(_mm256_and_si256(_mm256_castps_si256(VVM_cmpeq_AVX(v0, v1)), *(__m256i *)&v2))
+#define VVM_cmpne_logic_and_AVX(v0, v1, v2)     _mm256_castsi256_ps(_mm256_and_si256(_mm256_castps_si256(VVM_cmpneq_AVX(v0, v1)), *(__m256i *)&v2))
+#define VVM_cmplti_logic_and_AVX(v0, v1, v2)    _mm256_and_si256(VVM_cmplti_AVX(v0, v1), v2)
+#define VVM_cmplei_logic_and_AVX(v0, v1, v2)    _mm256_and_si256(VVM_cmplei_AVX(v0, v1), v2)
+#define VVM_cmpgti_logic_and_AVX(v0, v1, v2)    _mm256_and_si256(VVM_cmpgti_AVX(v0, v1), v2)
+#define VVM_cmpgei_logic_and_AVX(v0, v1, v2)    _mm256_and_si256(VVM_cmpgei_AVX(v0, v1), v2)
+#define VVM_cmpeqi_logic_and_AVX(v0, v1, v2)    _mm256_and_si256(VVM_cmpeqi_AVX(v0, v1), v2)
+#define VVM_cmpnei_logic_and_AVX(v0, v1, v2)    _mm256_and_si256(VVM_cmpneqi_AVX(v0, v1), v2)
+#define VVM_cmplt_logic_or_AVX(v0, v1, v2)      _mm256_castsi256_ps(_mm256_or_si256(_mm256_castps_si256(VVM_cmplt_AVX(v0, v1)), *(__m256i *)&v2))
+#define VVM_cmple_logic_or_AVX(v0, v1, v2)      _mm256_castsi256_ps(_mm256_or_si256(_mm256_castps_si256(VVM_cmple_AVX(v0, v1)), *(__m256i *)&v2))
+#define VVM_cmpgt_logic_or_AVX(v0, v1, v2)      _mm256_castsi256_ps(_mm256_or_si256(_mm256_castps_si256(VVM_cmpgt_AVX(v0, v1)), *(__m256i *)&v2))
+#define VVM_cmpge_logic_or_AVX(v0, v1, v2)      _mm256_castsi256_ps(_mm256_or_si256(_mm256_castps_si256(VVM_cmpge_AVX(v0, v1)), *(__m256i *)&v2))
+#define VVM_cmpeq_logic_or_AVX(v0, v1, v2)      _mm256_castsi256_ps(_mm256_or_si256(_mm256_castps_si256(VVM_cmpeq_AVX(v0, v1)), *(__m256i *)&v2))
+#define VVM_cmpne_logic_or_AVX(v0, v1, v2)      _mm256_castsi256_ps(_mm256_or_si256(_mm256_castps_si256(VVM_cmpneq_AVX(v0, v1)), *(__m256i *)&v2))
+#define VVM_cmplti_logic_or_AVX(v0, v1, v2)     _mm256_or_si256(VVM_cmplti_AVX(v0, v1), v2)
+#define VVM_cmplei_logic_or_AVX(v0, v1, v2)     _mm256_or_si256(VVM_cmplei_AVX(v0, v1), v2)
+#define VVM_cmpgti_logic_or_AVX(v0, v1, v2)     _mm256_or_si256(VVM_cmpgti_AVX(v0, v1), v2)
+#define VVM_cmpgei_logic_or_AVX(v0, v1, v2)     _mm256_or_si256(VVM_cmpgei_AVX(v0, v1), v2)
+#define VVM_cmpeqi_logic_or_AVX(v0, v1, v2)     _mm256_or_si256(VVM_cmpeqi_AVX(v0, v1), v2)
+#define VVM_cmpnei_logic_or_AVX(v0, v1, v2)     _mm256_or_si256(VVM_cmpneqi_AVX(v0, v1), v2)
+#define VVM_mad_add_AVX(v0, v1, v2, v3)         _mm256_add_ps(_mm256_fmadd_ps(v0, v1, v2), v3)
+#define VVM_mad_sub0_AVX(v0, v1, v2, v3)        _mm256_sub_ps(_mm256_fmadd_ps(v0, v1, v2), v3)
+#define VVM_mad_sub1_AVX(v0, v1, v2, v3)        _mm256_sub_ps(v3, _mm256_fmadd_ps(v0, v1, v2))
+#define VVM_mad_mul_AVX(v0, v1, v2, v3)         _mm256_mul_ps(_mm256_fmadd_ps(v0, v1, v2), v3)
+#define VVM_mad_sqrt_AVX(v0, v1, v2)            _mm256_sqrt_ps(_mm256_fmadd_ps(v0, v1, v2))
+#define VVM_mad_mad0_AVX(v0, v1, v2, v3, v4)    _mm256_fmadd_ps(_mm256_fmadd_ps(v0, v1, v2), v3, v4)
+#define VVM_mad_mad1_AVX(v0, v1, v2, v3, v4)    _mm256_fmadd_ps(v3, v4, _mm256_fmadd_ps(v0, v1, v2))
+#define VVM_mul_mad0_AVX(v0, v1, v2, v3)        _mm256_fmadd_ps(_mm256_mul_ps(v0, v1), v2, v3)
+#define VVM_mul_mad1_AVX(v0, v1, v2, v3)        _mm256_fmadd_ps(v2, v3, _mm256_mul_ps(v0, v1))
+#define VVM_mul_add_AVX(v0, v1, v2)             _mm256_add_ps(_mm256_mul_ps(v0, v1), v2)
+#define VVM_mul_sub_AVX(v0, v1, v2)             _mm256_sub_ps(_mm256_mul_ps(v0, v1), v2)
+#define VVM_mul_mul_AVX(v0, v1, v2)             _mm256_mul_ps(_mm256_mul_ps(v0, v1), v2)
+#define VVM_mul_max_AVX(v0, v1, v2)             _mm256_max_ps(_mm256_mul_ps(v0, v1), v2)
+#define VVM_add_mad1_AVX(v0, v1, v2, v3)        _mm256_fmadd_ps(v2, v3, _mm256_add_ps(v0, v1))
+#define VVM_add_add_AVX(v0, v1, v2)             _mm256_add_ps(_mm256_add_ps(v0, v1), v2)
+#define VVM_sub_cmplt1_AVX(v0, v1, v2)			VVM_cmplt_AVX(v2, _mm256_sub_ps(v0, v1))
+#define VVM_sub_neg_AVX(v0, v1)                 VVM_neg_AVX(_mm256_sub_ps(v0, v1))
+#define VVM_sub_mul_AVX(v0, v1, v2)				_mm256_mul_ps(_mm256_sub_ps(v0, v1), v2)
+#define VVM_div_mad0_AVX(v0, v1, v2, v3)        _mm256_fmadd_ps(_mm256_div_ps(v0, v1), v2, v3)
+#define VVM_div_f2i_AVX(v0, v1)                 _mm256_cvttps_epi32(_mm256_div_ps(_mm256_castsi256_ps(v0), _mm256_castsi256_ps(v1)))
+#define VVM_div_mul_AVX(v0, v1, v2)				_mm256_mul_ps(_mm256_div_ps(v0, v1), v2)
+#define VVM_muli_addi_AVX(v0, v1, v2)           _mm256_add_epi32(VectorIntMultiply_AVX(v1, v2), v2)
+#define VVM_addi_bit_rshift_AVX(v0, v1, v2)     _mm256_srlv_epi32(_mm256_add_epi32(v1, v2), v2)
+#define VVM_addi_muli_AVX(v0, v1, v2)           VectorIntMultiply_AVX(_mm256_add_epi32(v1, v2), v2)
+#define VVM_i2f_div0_AVX(v0, v1)                _mm256_div_ps(_mm256_cvtepi32_ps(_mm256_castps_si256(v0)), v1)
+#define VVM_i2f_div1_AVX(v0, v1)                _mm256_div_ps(v1, _mm256_cvtepi32_ps(_mm256_castps_si256(v0))) 
+#define VVM_i2f_mul_AVX(v0, v1)                 _mm256_mul_ps(_mm256_cvtepi32_ps(_mm256_castps_si256(v0)), v1)
+#define VVM_i2f_mad0_AVX(v0, v1, v2)            _mm256_fmadd_ps(_mm256_cvtepi32_ps(_mm256_castps_si256(v0)), v1, v2)
+#define VVM_i2f_mad1_AVX(v0, v1, v2)            _mm256_fmadd_ps(v0, v1, _mm256_cvtepi32_ps(_mm256_castps_si256(v2)))
+#define VVM_mul_sub0_AVX(v0, v1, v2)            _mm256_sub_ps(_mm256_mul_ps(v0, v1), v2)
+#define VVM_mul_sub1_AVX(v0, v1, v2)            _mm256_sub_ps(v2, _mm256_mul_ps(v0, v1))
+#define VVM_f2i_select1_AVX(mask, v0, v1)       _mm256_blendv_epi8(v1, _mm256_cvtps_epi32(_mm256_castsi256_ps(v0)), mask)
+#define VVM_f2i_maxi_AVX(v0, v1)                _mm256_max_epi32(_mm256_cvtps_epi32(_mm256_castsi256_ps(v0)), v1)
+#define VVM_f2i_addi_AVX(v0, v1)                _mm256_min_epi32(_mm256_cvtps_epi32(_mm256_castsi256_ps(v0)), v1)
+#define VVM_fmod_add_AVX(v0, v1, v2)            _mm256_add_ps(VectorMod_AVX(v0, v1), v2)
+#define VVM_bit_and_i2f_AVX(v0, v1)             _mm256_cvtepi32_ps(_mm256_and_si256(_mm256_castps_si256(v0), _mm256_castps_si256(v1)))
+#define VVM_bit_rshift_bit_and_AVX(v0, v1, v2)  _mm256_and_si256(VVMIntRShift_AVX(v0, v1), v2)
+#define VVM_neg_cmplt_AVX(v0, v1)               VectorCompareLT_AVX(VectorNegate_AVX(v0), v1)
+#define VVM_bit_or_muli_AVX(v0, v1, v2)         VectorIntMultiply_AVX(_mm256_or_si256(v0, v1), v2)
+#define VVM_bit_lshift_bit_or_AVX(v0, v1, v2)   _mm256_or_si256(VVMIntLShift_AVX(v0, v1), v2)
+#define VVM_random_add_AVX(v0, v1)              _mm256_add_ps(VVM_random_AVX(v0), v1)
+#define VVM_max_f2i_AVX(v0, v1)                 _mm256_cvtps_epi32(VectorMax_AVX(_mm256_castsi256_ps(v0), _mm256_castsi256_ps(v1)))
+#define VVM_select_mul_AVX(v0, v1, v2, v3)      _mm256_mul_ps(VectorSelect_AVX(v0, v1, v2), v3)
+#define VVM_select_add_AVX(v0, v1, v2, v3)      _mm256_add_ps(VectorSelect_AVX(v0, v1, v2), v3)
+
+
+//instructions called by the VM, named to match UE's style
+#define VectorAdd_AVX(v0, v1)                   _mm256_add_ps(v0, v1)
+#define VectorSubtract_AVX(v0, v1)              _mm256_sub_ps(v0, v1)
+#define VectorMultiply_AVX(v0, v1)              _mm256_mul_ps(v0, v1)
+#define VectorMultiplyAdd_AVX(v0, v1, v2)       _mm256_fmadd_ps(v0, v1, v2)
+#define VectorAbs_AVX(v)                        VVM_abs_AVX(v)
+#define VectorNegate_AVX(v)                     VVM_neg_AVX(v)
+#define VectorExp_AVX(v)                        exp256_ps(v)
+#define VectorExp2_AVX(v)                       _mm256_exp2_ps(v)
+#define VectorLog2_AVX(v)                       log256_ps(v)
+#define VectorSin_AVX(v)                        sin256_ps(v)
+#define VectorCos_AVX(v)                        cos256_ps(v)
+#define VectorLerp_AVX(v0, v1, v2)              _mm256_fmadd_ps(v1, v2, _mm256_mul_ps(v0, _mm256_sub_ps(VVM_m256Const(One), v2)))
+#define VectorCeil_AVX(v)                       _mm256_ceil_ps(v)
+#define VectorFloor_AVX(v)                      _mm256_floor_ps(v)
+#define VectorTruncate_AVX(v)                   _mm256_round_ps(v, _MM_FROUND_TO_NEG_INF |_MM_FROUND_NO_EXC)
+#define VectorFractional_AVX(v)                 _mm256_sub_ps(v, VectorTruncate_AVX(v))
+#define VectorClamp_AVX(v0, v1, v2)             _mm256_min_ps(_mm256_max_ps(v0, v1), v2)
+#define VectorMin_AVX(v0, v1)                   _mm256_min_ps(v0, v1)
+#define VectorMax_AVX(v0, v1)                   _mm256_max_ps(v0, v1)
+#define VectorRound_AVX(v)                      _mm256_round_ps(v, _MM_FROUND_TO_NEAREST_INT |_MM_FROUND_NO_EXC)
+#define VectorSign_AVX(v)                       _mm256_blendv_ps(VVM_m256Const(NegativeOne), VVM_m256Const(One), _mm256_cmp_ps(v, _mm256_setzero_ps(), _CMP_GT_OS))
+#define VectorSelect_AVX(m, v0, v1)				VVM_select_AVX(m, v0, v1)
+#define VectorCompareLT_AVX(v0, v1)             VVM_cmplt_AVX(v0, v1)
+#define VectorCompareLE_AVX(v0, v1)             VVM_cmple_AVX(v0, v1)
+#define VectorCompareGT_AVX(v0, v1)             VVM_cmpgt_AVX(v0, v1)
+#define VectorCompareGE_AVX(v0, v1)             VVM_cmpge_AVX(v0, v1)
+#define VectorCompareEQ_AVX(v0, v1)             VVM_cmpeq_AVX(v0, v1)
+#define VectorCompareNE_AVX(v0, v1)             VVM_cmpneq_AVX(v0, v1)
+#define VectorIntAdd_AVX(v0, v1)                _mm256_add_epi32(v0, v1)
+#define VectorIntSubtract_AVX(v0, v1)           _mm256_sub_epi32(v0, v1)
+#define VectorIntMultiply_AVX(v0, v1)           _mm256_unpacklo_epi32(_mm256_shuffle_epi32(_mm256_mul_epu32(v0, v1), _MM_SHUFFLE(0, 0, 2, 0)), _mm256_shuffle_epi32(_mm256_mul_epu32(_mm256_srli_si256(v0, 4), _mm256_srli_si256(v1, 4)), _MM_SHUFFLE(0, 0, 2, 0)))
+#define VectorIntClamp_AVX(v0, v1, v2)          _mm256_min_epi32(_mm256_max_epi32(v0, v1), v2)
+#define VectorIntMin_AVX(v0, v1)                _mm256_min_epi32(v0, v1)
+#define VectorIntMax_AVX(v0, v1)                _mm256_max_epi32(v0, v1)
+#define VectorIntAbs_AVX(v)                     _mm256_abs_epi32(v)
+#define VectorIntNegate_AVX(v)                  _mm256_sub_epi32(_mm256_setzero_si256(), v)
+#define VectorIntSign_AVX(v)                    _mm256_blendv_epi8(VVM_m256iConst(NegativeOne), VVM_m256iConst(One), VVM_cmpgei_AVX(v, _mm256_setzero_si256()))
+#define VectorIntCompareLT_AVX(v0, v1)          VVM_cmplti_AVX(v0, v1)
+#define VectorIntCompareLE_AVX(v0, v1)          VVM_cmplei_AVX(v0, v1)
+#define VectorIntCompareGT_AVX(v0, v1)          VVM_cmpgti_AVX(v0, v1)
+#define VectorIntCompareGE_AVX(v0, v1)          VVM_cmpgei_AVX(v0, v1)
+#define VectorIntCompareEQ_AVX(v0, v1)          VVM_cmpeqi_AVX(v0, v1)
+#define VectorIntCompareNEQ_AVX(v0, v1)         VVM_cmpneqi_AVX(v0, v1)
+#define VectorIntAnd_AVX(v0, v1)                _mm256_and_si256(v0, v1)
+#define VectorIntOr_AVX(v0, v1)                 _mm256_or_si256(v0, v1)
+#define VectorIntXor_AVX(v0, v1)                _mm256_xor_si256(v0, v1)
+#define VectorIntNot_AVX(v)                     _mm256_xor_si256(v, VVM_m256iConst(FMask))
+#define VVMIntLShift_AVX(v0, v1)                _mm256_sllv_epi32(v0, v1)
+#define VVMIntRShift_AVX(v0, v1)                _mm256_srlv_epi32(v0, v1)
+#define VVMf2i_AVX(v)                           _mm256_cvttps_epi32(_mm256_castsi256_ps(v))
+#define VVMi2f_AVX(v)                           _mm256_cvtepi32_ps(_mm256_castps_si256(v))
+#define VectorSinCos_AVX(v, OutSin, OutCos)     sincos256_ps(v, OutSin, OutCos)
+
+static inline __m256 VectorTan_AVX(const __m256 &v)
+{
+	VVMXXMYMMUnion U;
+	U.ymm = v;
+	U.xmm[0] = VectorTan(U.xmm[0]);
+	U.xmm[1] = VectorTan(U.xmm[1]);
+	return U.ymm;
+}
+
+static inline __m256 VectorATan_AVX(const __m256 &v)
+{
+	VVMXXMYMMUnion U;
+	U.ymm = v;
+	U.xmm[0] = VectorATan(U.xmm[0]);
+	U.xmm[1] = VectorATan(U.xmm[1]);
+	return U.ymm;
+}
+
+static inline __m256 VectorATan2_AVX(const __m256 &v0, const __m256 &v1)
+{
+	VVMXXMYMMUnion U0, U1;
+	U0.ymm = v0;
+	U1.ymm = v1;
+	U0.xmm[0] = VectorATan2(U0.xmm[0], U1.xmm[0]);
+	U0.xmm[1] = VectorATan2(U0.xmm[1], U1.xmm[1]);
+	return U0.ymm;
+}
+
+static inline __m256 VectorASin_AVX(const __m256 &v)
+{
+	VVMXXMYMMUnion U;
+	U.ymm = v;
+	U.xmm[0] = VectorASin(U.xmm[0]);
+	U.xmm[1] = VectorASin(U.xmm[1]);
+	return U.ymm;
+}
+
+static inline __m256 VectorMod_AVX(const __m256 &v0, const __m256 &v1)
+{
+	VVMXXMYMMUnion U0, U1;
+	U0.ymm = v0;
+	U1.ymm = v1;
+	U0.xmm[0] = VectorMod(U0.xmm[0], U1.xmm[0]);
+	U0.xmm[1] = VectorMod(U0.xmm[1], U1.xmm[1]);
+	return U0.ymm;
+}
+
+struct FVVM_VU8
+{
+	union
+	{
+		__m256i i8;
+		uint32 o8[8];
+	}
+};
+
+static inline __m256i VVMIntDiv_AVX(const __m256i &X, const __m256i &Y)
+{
+	uint32 *v0_8 = (uint32 *)&X;
+	uint32 *v1_8 = (uint32 *)&Y;
+
+    FVVM_VU8 res;
+
+	res.o8[0] = v1_8[0] == 0 ? 0 : (v0_8[0] / v1_8[0]);
+    res.o8[1] = v1_8[1] == 0 ? 0 : (v0_8[1] / v1_8[1]);
+    res.o8[2] = v1_8[2] == 0 ? 0 : (v0_8[2] / v1_8[2]);
+    res.o8[3] = v1_8[3] == 0 ? 0 : (v0_8[3] / v1_8[3]);
+	res.o8[4] = v1_8[4] == 0 ? 0 : (v0_8[4] / v1_8[4]);
+    res.o8[5] = v1_8[5] == 0 ? 0 : (v0_8[5] / v1_8[5]);
+    res.o8[6] = v1_8[6] == 0 ? 0 : (v0_8[6] / v1_8[6]);
+    res.o8[7] = v1_8[7] == 0 ? 0 : (v0_8[7] / v1_8[7]);
+	
+	return res.i8;
+}
+
+
+#endif //VECTORVM_SUPPORTS_AVX
+
+struct FVVM_VU4
+{
+	union
+	{
+		VectorRegister4i i4;
+		uint32 o4[4];
+	};
+};
+
+inline VectorRegister4i VVMIntRShift(VectorRegister4i v0, VectorRegister4i v1)
+{
+#if PLATFORM_ENABLE_VECTORINTRINSICS_NEON
+	VectorRegister4i res = vshlq_u32(v0, vmulq_s32(v1, vdupq_n_s32(-1)));
+	return res;
+#else
+	uint32 *v0_4 = (uint32 *)&v0;
+	uint32 *v1_4 = (uint32 *)&v1;
+
+	FVVM_VU4 res;
+
+	res.o4[0] = v0_4[0] >> v1_4[0];
+    res.o4[1] = v0_4[1] >> v1_4[1];
+    res.o4[2] = v0_4[2] >> v1_4[2];
+    res.o4[3] = v0_4[3] >> v1_4[3];
+	
+	return res.i4;
+#endif
+}
+
+inline VectorRegister4i VVMIntLShift(VectorRegister4i v0, VectorRegister4i v1)
+{
+#if PLATFORM_ENABLE_VECTORINTRINSICS_NEON
+	VectorRegister4i res = vshlq_u32(v0, v1);
+	return res;
+#else
+	uint32 *v0_4 = (uint32 *)&v0;
+	uint32 *v1_4 = (uint32 *)&v1;
+    
+	FVVM_VU4 res;
+
+	res.o4[0] = v0_4[0] << v1_4[0];
+    res.o4[1] = v0_4[1] << v1_4[1];
+    res.o4[2] = v0_4[2] << v1_4[2];
+    res.o4[3] = v0_4[3] << v1_4[3];
+	
+	return res.i4;
+#endif
+}
+
+inline VectorRegister4i VVMIntDiv(VectorRegister4i v0, VectorRegister4i v1)
+{
+	uint32 *v0_4 = (uint32 *)&v0;
+	uint32 *v1_4 = (uint32 *)&v1;
+
+    FVVM_VU4 res;
 
+	res.o4[0] = v1_4[0] == 0 ? 0 : (v0_4[0] / v1_4[0]);
+    res.o4[1] = v1_4[1] == 0 ? 0 : (v0_4[1] / v1_4[1]);
+    res.o4[2] = v1_4[2] == 0 ? 0 : (v0_4[2] / v1_4[2]);
+    res.o4[3] = v1_4[3] == 0 ? 0 : (v0_4[3] / v1_4[3]);
+	
+	return res.i4;
+}
+
+inline VectorRegister4i VVMf2i(VectorRegister4i v0)
+{
+	FVecReg u;
+	u.i = v0;
+	VectorRegister4i res = VectorFloatToInt(u.v);
+	return res;
+}
+
+inline VectorRegister4f VVMi2f(VectorRegister4f v0)
+{
+	FVecReg u;
+	u.v = v0;
+	VectorRegister4f res = VectorIntToFloat(u.i);
+	return res;
+}
 #define VVMDebugBreakIf(expr)	if ((expr)) { PLATFORM_BREAK(); }
 
-static void VVMMemCpy(void *dst, void *src, size_t bytes)
+FORCEINLINE uint16 float_to_half_fast3_rtne(uint32 f_in)
+{
+	uint16 h_out;
+	FPlatformMath::StoreHalf(&h_out, *reinterpret_cast<float*>(&f_in));
+
+	return h_out;
+}
+
+FORCEINLINE float half_to_float(uint16 h_)
+{
+	return FPlatformMath::LoadHalf(&h_);
+}
+
+FORCEINLINE void VVMMemCpy(void *dst, void *src, size_t bytes)
 {
-	unsigned char *d      = (unsigned char *)dst;
-	unsigned char *s      = (unsigned char *)src;
-	unsigned char *s_end  = s + bytes;
+	unsigned char *RESTRICT d      = (unsigned char *)dst;
+	unsigned char *RESTRICT s      = (unsigned char *)src;
+	unsigned char *RESTRICT s_end  = s + bytes;
 	ptrdiff_t ofs_to_dest = d - s;
 	if (bytes < 16)
 	{
@@ -371,77 +849,165 @@ static void VVMMemCpy(void *dst, void *src, size_t bytes)
 	}      
 }
 
-static void VVMMemSet32(void *dst, uint32 val, size_t num_vals)
+FORCEINLINE void VVMMemSet32(void *dst, uint32 val, size_t num_vals)
 {
-	check(((size_t)dst & 3) == 0);							    //must be 4 byte aligned
-	if (num_vals < 4)
+	if (num_vals <= 4)
 	{
-		uint32 *ptr = (uint32 *)dst;
-		uint32 *end_ptr = ptr + num_vals;
-		do
+		uint32 *dst32 = (uint32 *)dst;
+		switch (num_vals)
 		{
-			*ptr++ = val;
-		} while (ptr < end_ptr);
+			case 4:		VectorIntStore(VectorIntSet1(val), dst); break;
+			case 3:		dst32[2] = val;  //intentional fallthrough
+			case 2:		dst32[1] = val;  //intentional fallthrough
+			case 1:		dst32[0] = val;  //intentional fallthrough
+			case 0:		break;
+		}
 	}
 	else
 	{
 		VectorRegister4i v4 = VectorIntSet1(val);
-		check(((size_t)dst & 3) == 0);
-		uint32 *ptr = (uint32 *)dst;
-		uint32 *end_ptr = ptr + num_vals - 4;
+		char *RESTRICT ptr = (char *)dst;
+		char *RESTRICT end_ptr = ptr + num_vals * 4 - 16;
 		while (ptr < end_ptr) {
 			VectorIntStore(v4, ptr);
-			ptr += 4;
+			ptr += 16;
 		}
 		VectorIntStore(v4, end_ptr);
 	}
 }
+#pragma warning(push)
+#pragma warning(disable: 4426)
+#include "VectorVMExperimental_Optimizer.inl" 
+#include "VectorVMExperimental_Memory.inl"
+#pragma warning(pop)
+
+#if VECTORVM_SUPPORTS_AVX
+static void SetRegIncTableFor8WideFrom4Wide(FVectorVMExecContext *ExecCtx, FVectorVMBatchState *BatchState)
+{
+	uint8 *TempRegIncTable  = BatchState->RegIncTable;
+	uint8 *InputRegIncTable = BatchState->RegIncTable + ExecCtx->VVMState->NumTempRegisters + ExecCtx->VVMState->NumConstBuffers;
+	for (uint32 i = 0; i < ExecCtx->VVMState->NumTempRegisters; ++i)
+	{
+		TempRegIncTable[i] <<= 1;
+	}
+	for (uint32 i = 0; i < ExecCtx->VVMState->NumInputBuffers; ++i)
+	{
+		InputRegIncTable[i] <<= 1;
+	}
+}
 
-#include "./VectorVMExperimental_Optimizer.inl" 
-
-static uint8 *SetupBatchStatePtrs(FVectorVMState *VVMState, FVectorVMBatchState *BatchState, void *BatchMem)
+static void SetRegIncTableFor4WideFrom8Wide(FVectorVMExecContext *ExecCtx, FVectorVMBatchState *BatchState)
+{
+	uint8 *TempRegIncTable  = BatchState->RegIncTable;
+	uint8 *InputRegIncTable = BatchState->RegIncTable + ExecCtx->VVMState->NumTempRegisters + ExecCtx->VVMState->NumConstBuffers;
+	for (uint32 i = 0; i < ExecCtx->VVMState->NumTempRegisters; ++i)
+	{
+		TempRegIncTable[i] >>= 1;
+	}
+	for (uint32 i = 0; i < ExecCtx->VVMState->NumInputBuffers; ++i)
+	{
+		InputRegIncTable[i] >>= 1;
+	}
+}
+#endif //VECTORVM_SUPPORTS_AVX
+static uint8 *SetupBatchStatePtrs(FVectorVMExecContext *ExecCtx, FVectorVMBatchState *BatchState)
 {
-	uint8 *BatchDataPtr = (uint8 *)VVM_ALIGN_CACHELINE(BatchMem);
-	BatchState->MallocedMemPtr                              = BatchMem;
-	BatchState->RegisterData                                = (FVecReg *)BatchDataPtr;                  BatchDataPtr += VVMState->PerBatchRegisterDataBytesRequired;
-	BatchState->ChunkLocalData.StartingOutputIdxPerDataSet	= (uint32 *) BatchDataPtr;                  BatchDataPtr += VVMState->PerBatchChunkLocalDataOutputIdxBytesRequired;
-	BatchState->ChunkLocalData.NumOutputPerDataSet          = (uint32 *) BatchDataPtr;                  BatchDataPtr += VVMState->PerBatchChunkLocalNumOutputBytesRequired;
+	uint8 *BatchDataPtr        = (uint8 *) VVM_ALIGN_64((size_t)BatchState + sizeof(FVectorVMBatchState));
+	size_t NumPtrRegsInTable   = ExecCtx->VVMState->NumTempRegisters + ExecCtx->VVMState->NumConstBuffers + ExecCtx->VVMState->NumInputBuffers * 2;
+
+	BatchState->RegisterData                                = (FVecReg *)BatchDataPtr;                  BatchDataPtr += ExecCtx->Internal.PerBatchRegisterDataBytesRequired;
+	BatchState->RegPtrTable                                 = (uint8 **)BatchDataPtr;                   BatchDataPtr += NumPtrRegsInTable * sizeof(uint32 *);
+	BatchState->RegIncTable                                 = (uint8 *)BatchDataPtr;                    BatchDataPtr += NumPtrRegsInTable;
+	BatchState->ChunkLocalData.StartingOutputIdxPerDataSet	= (uint32 *) BatchDataPtr;                  BatchDataPtr += ExecCtx->VVMState->ChunkLocalDataOutputIdxNumBytes;
+	BatchState->ChunkLocalData.NumOutputPerDataSet          = (uint32 *) BatchDataPtr;                  BatchDataPtr += ExecCtx->VVMState->ChunkLocalNumOutputNumBytes;
 
 	{ //deal with the external function register decoding buffer
 		size_t PtrBeforeExtFnDecodeReg = (size_t)BatchDataPtr;
-		BatchState->ChunkLocalData.ExtFnDecodedReg.RegData       = (FVecReg **)BatchDataPtr;  BatchDataPtr += sizeof(FVecReg *) * VVMState->MaxExtFnRegisters;
-		BatchState->ChunkLocalData.ExtFnDecodedReg.RegInc        = (uint32 *)BatchDataPtr;    BatchDataPtr += sizeof(uint32)    * VVMState->MaxExtFnRegisters;
-		BatchDataPtr = (uint8 *)VVM_ALIGN_16(BatchDataPtr);
-		BatchState->ChunkLocalData.ExtFnDecodedReg.DummyRegs     = (FVecReg *)BatchDataPtr;   BatchDataPtr += sizeof(FVecReg)   * VVMState->NumDummyRegsReq;
+		BatchState->ChunkLocalData.ExtFnDecodedReg.RegData       = (uint32 **)BatchDataPtr;  BatchDataPtr += sizeof(FVecReg *) * ExecCtx->VVMState->MaxExtFnRegisters;
+		BatchState->ChunkLocalData.ExtFnDecodedReg.RegInc        = (uint8 *)BatchDataPtr;    BatchDataPtr += sizeof(uint8)     * ExecCtx->VVMState->MaxExtFnRegisters;
+		BatchDataPtr = (uint8 *)VVM_PTR_ALIGN(BatchDataPtr);
+		BatchState->ChunkLocalData.ExtFnDecodedReg.DummyRegs     = (FVecReg *)BatchDataPtr;  BatchDataPtr += sizeof(FVecReg)   * ExecCtx->VVMState->NumDummyRegsRequired;
 		size_t PtrAfterExtFnDecodeReg = (size_t)BatchDataPtr;
-		check(PtrAfterExtFnDecodeReg - PtrBeforeExtFnDecodeReg <= VVMState->PerBatchChunkLocalNumExtFnDecodeRegisterBytesRequired);
 	}
-	
-	{ //after everything, likely outside of what the chunk will cache, setup the almost-never-used random counters and
-		BatchState->ChunkLocalData.RandCounters = (int32 *)BatchDataPtr;
-		BatchDataPtr += sizeof(int32) * VVMState->MaxInstancesPerChunk;
+
+	{ //build the register pointer table which contains pointers (in order) to:
+		uint32 **TempRegPtr     = (uint32 **)BatchState->RegPtrTable;                  //1. Temp Registers
+		uint32 **ConstBuffPtr   = TempRegPtr   + ExecCtx->VVMState->NumTempRegisters;  //2. Constant Buffers
+		uint32 **InputPtr       = ConstBuffPtr + ExecCtx->VVMState->NumConstBuffers;   //3. Input Registers (the 2x is because 
+		uint32 NumLoops         = ExecCtx->Internal.MaxInstancesPerChunk >> 2;
+
+		static_assert(sizeof(FVecReg) == 16);
+		FMemory::Memset(BatchState->RegIncTable, sizeof(FVecReg), NumPtrRegsInTable);
+
+		//temp regsiters
+		for (uint32 i = 0; i < ExecCtx->VVMState->NumTempRegisters; ++i) {
+			TempRegPtr[i] = (uint32 *)(BatchState->RegisterData + i * NumLoops);
+		}
+		//constant buffers
+		for (uint32 i = 0; i < ExecCtx->VVMState->NumConstBuffers; ++i) {
+			ConstBuffPtr[i] = (uint32 *)(ExecCtx->VVMState->ConstantBuffers + i);
+			BatchState->RegIncTable[ExecCtx->VVMState->NumTempRegisters + i] = 0;
+		}
+		//inputs
+		int NoAdvCounter = 0;
+		for (uint32 i = 0; i < ExecCtx->VVMState->NumInputBuffers; ++i)
+		{
+			uint8 DataSetIdx      = ExecCtx->VVMState->InputMapCacheIdx[i];
+			uint16 InputMapSrcIdx = ExecCtx->VVMState->InputMapCacheSrc[i];
+
+			uint32 **DataSetInputBuffers = (uint32 **)ExecCtx->DataSets[DataSetIdx].InputRegisters.GetData();
+			int32 InstanceOffset         = ExecCtx->DataSets[DataSetIdx].InstanceOffset;
+
+			if (InputMapSrcIdx & 0x8000) //this is a noadvance input.  It points to data after the constant buffers
+			{
+				InputPtr[i                                     ] = (uint32 *)(ExecCtx->VVMState->ConstantBuffers + ExecCtx->VVMState->NumConstBuffers + NoAdvCounter);
+				InputPtr[i + ExecCtx->VVMState->NumInputBuffers] = (uint32 *)(ExecCtx->VVMState->ConstantBuffers + ExecCtx->VVMState->NumConstBuffers + NoAdvCounter);
+				++NoAdvCounter;
+				
+				BatchState->RegIncTable[ExecCtx->VVMState->NumTempRegisters + ExecCtx->VVMState->NumConstBuffers + i] = 0; //no advance inputs... don't advance obviously
+
+				if (InputMapSrcIdx & 0x4000) //half input (@TODO: has never been tested)
+				{
+					uint16 *Ptr = (uint16 *)DataSetInputBuffers[InputMapSrcIdx & 0x3FFF] + InstanceOffset;
+					float val = half_to_float(*Ptr);
+					VectorRegister4f InputVal = VectorSet1(val);
+					VectorStore(InputVal, (float *)InputPtr[i]);
+				}
+				else
+				{
+					uint32 *Ptr = (uint32 *)DataSetInputBuffers[InputMapSrcIdx & 0x3FFF] + InstanceOffset;
+					VectorRegister4i InputVal4 = VectorIntSet1(*Ptr);
+					VectorIntStore(InputVal4, InputPtr[i]);
+				}
+			}
+			else //regular input, point directly to the input buffer
+			{
+				InputPtr[i                                     ] = DataSetInputBuffers[InputMapSrcIdx] + InstanceOffset;
+				InputPtr[i + ExecCtx->VVMState->NumInputBuffers] = InputPtr[i]; //second copy of the "base" ptr so each chunk can start them at their correct starting offset
+			}
+		}
 	}
 
-	if ((size_t)(BatchDataPtr - (uint8 *)BatchMem) <= VVMState->NumBytesRequiredPerBatch)
+	BatchState->ChunkLocalData.RandCounters = nullptr; //these get malloc'd separately if they're ever used... which they very rarely are
+
+	if ((size_t)(BatchDataPtr - (uint8 *)BatchState) <= ExecCtx->Internal.NumBytesRequiredPerBatch)
 	{
 		return BatchDataPtr;
 	}
 	else
 	{
-		VVMState->Error.Flags |= VVMErr_BatchMemory;
-		VVMState->Error.LineNum = __LINE__;
 		return nullptr;
 	}
 }
 
-static bool AssignInstancesToBatch(FVectorVMState *VVMState, FVectorVMBatchState *BatchState)
+inline bool AssignInstancesToBatch(FVectorVMExecContext *ExecCtx, FVectorVMBatchState *BatchState)
 {
 	int SanityCount = 0;
 	do
 	{
-		int OldNumAssignedInstances = VVMState->NumInstancesAssignedToBatches;
-		int MaxInstancesPerBatch    = VVMState->MaxInstancesPerChunk * VVMState->MaxChunksPerBatch;
-		int NumAssignedInstances    = FPlatformAtomics::InterlockedCompareExchange(&VVMState->NumInstancesAssignedToBatches, OldNumAssignedInstances + MaxInstancesPerBatch, OldNumAssignedInstances);
+		int OldNumAssignedInstances = ExecCtx->Internal.NumInstancesAssignedToBatches;
+		int MaxInstancesPerBatch    = ExecCtx->Internal.MaxInstancesPerChunk * ExecCtx->Internal.MaxChunksPerBatch;
+		int NumAssignedInstances    = FPlatformAtomics::InterlockedCompareExchange(&ExecCtx->Internal.NumInstancesAssignedToBatches, OldNumAssignedInstances + MaxInstancesPerBatch, OldNumAssignedInstances);
 		if (NumAssignedInstances == OldNumAssignedInstances)
 		{
 			//we allow "overflow" on the thread-local BatchState to keep access to the contended variable: VVMState->NumInstancesAssignedToBatches 
@@ -449,14 +1015,12 @@ static bool AssignInstancesToBatch(FVectorVMState *VVMState, FVectorVMBatchState
 			//for this BatchState.
 			BatchState->StartInstance = OldNumAssignedInstances;
 			BatchState->NumInstances  = MaxInstancesPerBatch;
-			if (BatchState->StartInstance + BatchState->NumInstances > VVMState->TotalNumInstances) {
-				BatchState->NumInstances = VVMState->TotalNumInstances - BatchState->StartInstance;
+			if (BatchState->StartInstance + BatchState->NumInstances > ExecCtx->NumInstances) {
+				BatchState->NumInstances = ExecCtx->NumInstances - BatchState->StartInstance;
 			}
 			if (BatchState->NumInstances <= 0) {
 				return false; //some other thread interrupted and finished the rest of the instances, we're done.
 			}
-			VVMMemSet32(BatchState->ChunkLocalData.RandCounters, 0, VVMState->MaxInstancesPerChunk); //@TODO: complete waste.  With a more sensible prng we could remove this.
-			++BatchState->UseCount;
 			return true;
 		}
 	} while (SanityCount++ < (1 << 30));
@@ -466,16 +1030,10 @@ static bool AssignInstancesToBatch(FVectorVMState *VVMState, FVectorVMBatchState
 
 VECTORVM_API void FreeVectorVMState(FVectorVMState *VVMState)
 {
-	if (VVMState == nullptr)
-	{
-		return;
-	}
-	VectorVMFreeFn *FreeFn = VVMState->FreeFn ? VVMState->FreeFn : VVMDefaultFree;
-	for (int i = 1; i < VVMState->NumBatches; ++i)
+	if (VVMState != nullptr)
 	{
-		FreeFn(VVMState->BatchStates[i].MallocedMemPtr, __FILE__, __LINE__);
+		FMemory::Free(VVMState);
 	}
-	FreeFn(VVMState, __FILE__, __LINE__);
 }
 
 static void SetupRandStateForBatch(FVectorVMBatchState *BatchState)
@@ -486,6 +1044,18 @@ static void SetupRandStateForBatch(FVectorVMBatchState *BatchState)
 	//use psuedo-pcg to setup a state for xorwow
 	for (int i = 0; i < 5; ++i) //loop for xorwow internal state
 	{
+#		if VECTORVM_SUPPORTS_AVX
+		MS_ALIGN(32) uint32 Values[8] GCC_ALIGN(32);
+		for (int j = 0; j < 8; ++j)
+		{
+			uint64 old_state   = pcg_state;
+			pcg_state          = old_state * 6364136223846793005ULL + pcg_inc;
+			uint32 xor_shifted = (uint32)(((old_state >> 18U) ^ old_state) >> 27U);
+			uint32 rot         = old_state >> 59U;
+			Values[j]          = (xor_shifted >> rot) | (xor_shifted << ((0U - rot) & 31));
+		}
+		_mm256_storeu_si256(BatchState->RandState.AVX.State + i, *(__m256i *)Values);
+#		else
 		MS_ALIGN(16) uint32 Values[4] GCC_ALIGN(16);
 		for (int j = 0; j < 4; ++j)
 		{
@@ -495,1648 +1065,426 @@ static void SetupRandStateForBatch(FVectorVMBatchState *BatchState)
 			uint32 rot         = old_state >> 59U;
 			Values[j]          = (xor_shifted >> rot) | (xor_shifted << ((0U - rot) & 31));
 		}
-		VectorIntStore(*(VectorRegister4i *)Values, BatchState->RandState + i);
+		VectorIntStore(*(VectorRegister4i *)Values, BatchState->RandState.State + i);
+#		endif
 	}
-	BatchState->RandCounters = MakeVectorRegisterInt64(pcg_inc, pcg_state);
+	BatchState->RandState.Counters = MakeVectorRegisterInt64(pcg_inc, pcg_state);
 	BatchState->RandStream.GenerateNewSeed();
 }
 
 static VectorRegister4i VVMXorwowStep(FVectorVMBatchState *BatchState)
 {
-	VectorRegister4i t = BatchState->RandState[4];
-	VectorRegister4i s = BatchState->RandState[0];
-	BatchState->RandState[4] = BatchState->RandState[3];
-	BatchState->RandState[3] = BatchState->RandState[2];
-	BatchState->RandState[2] = BatchState->RandState[1];
-	BatchState->RandState[1] = s;
+	VectorRegister4i t = BatchState->RandState.State[4];
+	VectorRegister4i s = BatchState->RandState.State[0];
+	BatchState->RandState.State[4] = BatchState->RandState.State[3];
+	BatchState->RandState.State[3] = BatchState->RandState.State[2];
+	BatchState->RandState.State[2] = BatchState->RandState.State[1];
+	BatchState->RandState.State[1] = s;
 	t = VectorIntXor(t, VectorShiftRightImmLogical(t, 2));
 	t = VectorIntXor(t, VectorShiftLeftImm(t, 1));
 	t = VectorIntXor(t, VectorIntXor(s, VectorIntXor(s, VectorShiftLeftImm(s, 4))));
-	BatchState->RandState[0] = t;
-	BatchState->RandCounters = VectorIntAdd(BatchState->RandCounters, VectorIntSet1(362437));
-	VectorRegister4i Result = VectorIntAdd(t, VectorIntLoad(&BatchState->RandCounters));
+	BatchState->RandState.State[0] = t;
+	BatchState->RandState.Counters = VectorIntAdd(BatchState->RandState.Counters, VectorIntSet1(362437));
+	VectorRegister4i Result = VectorIntAdd(t, VectorIntLoad(&BatchState->RandState.Counters));
 	return Result;
 }
-
-
-VECTORVM_API FVectorVMState *InitVectorVMState(FVectorVMInitData *InitData, FVectorVMExternalFnPerInstanceData **OutPerInstanceExtData, struct FVectorVMSerializeState *SerializeState)
+#if VECTORVM_SUPPORTS_AVX
+static __m256i VVMXorwowStep_AVX(FVectorVMBatchState *BatchState)
 {
-	if (InitData->OptimizeContext == nullptr)
-	{
+	__m256i t = BatchState->RandState.AVX.State[4];
+	__m256i s = BatchState->RandState.AVX.State[0];
+	BatchState->RandState.AVX.State[4] = BatchState->RandState.AVX.State[3];
+	BatchState->RandState.AVX.State[3] = BatchState->RandState.AVX.State[2];
+	BatchState->RandState.AVX.State[2] = BatchState->RandState.AVX.State[1];
+	BatchState->RandState.AVX.State[1] = s;
+	t = _mm256_xor_si256(t, _mm256_srli_epi32(t, 2));
+	t = _mm256_xor_si256(t, _mm256_slli_epi32(t, 1));
+	t = _mm256_xor_si256(t, _mm256_xor_si256(s, _mm256_xor_si256(s, _mm256_slli_epi32(s, 4))));
+	BatchState->RandState.AVX.State[0] = t;
+	BatchState->RandState.AVX.Counters = _mm256_add_epi32(BatchState->RandState.AVX.Counters, _mm256_set1_epi32(362437));
+	__m256i Result = _mm256_add_epi32(t, _mm256_loadu_si256(&BatchState->RandState.AVX.Counters));
+	return Result;
+}
+#endif
+
+VECTORVM_API FVectorVMState *AllocVectorVMState(FVectorVMOptimizeContext *OptimizeCtx) {
+	if (OptimizeCtx == nullptr || OptimizeCtx->Error.Line != 0) {
 		return nullptr;
 	}
-	if (InitData->ExtFunctionTable.Num() > 0 && InitData->OptimizeContext->MaxExtFnUsed >= InitData->ExtFunctionTable.Num())
+	//compute the number of overhead bytes for this VVM State
+	uint32 ConstBufferOffset         = VVM_ALIGN_32(sizeof(FVectorVMState));
+
+	size_t ConstantBufferNumBytes    = VVM_PTR_ALIGN(VVM_REG_SIZE                     * (OptimizeCtx->NumConstsRemapped + OptimizeCtx->NumNoAdvanceInputs)  );
+	size_t ExtFnTableNumBytes        = VVM_PTR_ALIGN(sizeof(FVectorVMExtFunctionData) * (OptimizeCtx->MaxExtFnUsed + 1)                                     );
+	size_t OutputPerDataSetNumBytes  = VVM_PTR_ALIGN(sizeof(volatile int32)           * OptimizeCtx->NumOutputDataSets                                      );
+	size_t ConstMapCacheNumBytes     = VVM_PTR_ALIGN((sizeof(uint8) + sizeof(uint16)) * OptimizeCtx->NumConstsRemapped                                      );
+	size_t InputMapCacheNumBytes     = VVM_PTR_ALIGN((sizeof(uint8) + sizeof(uint16)) * OptimizeCtx->NumInputsRemapped                                      );
+	size_t VVMStateTotalNumBytes     = ConstBufferOffset + ConstantBufferNumBytes + ExtFnTableNumBytes + OutputPerDataSetNumBytes + ConstMapCacheNumBytes + InputMapCacheNumBytes;
+
+	uint8 *StatePtr = (uint8 *)FMemory::Malloc(VVMStateTotalNumBytes, 16);
+	if (StatePtr == nullptr)
 	{
-		check(false); //somehow the funciton table changed in between optimize() and init()
 		return nullptr;
 	}
 
-	size_t NumVVMStateBytesRequired                                    = 0; //Memory "outside" of the batches for the VVM state itself
-	size_t PerBatchRegisterDataBytesRequired                           = 0;																									                    //const + temp buffer for each batch
-	const uint32 MaxExtFnRegisters                                     = InitData->OptimizeContext->MaxExtFnRegisters == 0 ? 0 : VVM_ALIGN_4(InitData->OptimizeContext->MaxExtFnRegisters + 3); //we decode 4 at a time, so if we need any, we need a multiuple of 4
-	const size_t PerBatchChunkLocalDataOutputIdxBytesRequired	       = sizeof(uint32)                                        * InitData->OptimizeContext->NumOutputDataSets;                  //chunk local bytes required for instance offset
-	const size_t PerBatchChunkLocalNumOutputBytesRequired		       = sizeof(uint32)                                        * InitData->OptimizeContext->NumOutputDataSets;                  //chunk local bytes for num outputs
-	const size_t ConstantBufferSize					       = sizeof(FVecReg)                                       * InitData->OptimizeContext->NumConstsRemapped;
-	const size_t PerBatchChunkLocalNumExtFnDecodeRegisterBytesRequired = (sizeof(FVecReg *) + sizeof(uint32)) * MaxExtFnRegisters + sizeof(FVecReg) * (InitData->OptimizeContext->NumDummyRegsReq + 1); //+1 for potential padding
-
-	const size_t BatchOverheadSize                                     = ConstantBufferSize + 
-		                                                                 PerBatchChunkLocalDataOutputIdxBytesRequired + 
-		                                                                 PerBatchChunkLocalNumOutputBytesRequired + 
-		                                                                 PerBatchChunkLocalNumExtFnDecodeRegisterBytesRequired +
-		                                                                 InitData->OptimizeContext->NumBytecodeBytes +
-		                                                                 VVM_CHUNK_FIXED_OVERHEAD_SIZE;
-	static const uint32 MaxChunksPerBatch                              = 4; //*MUST BE POW 2* arbitrary 4 chunks per batch... this is harder to load balance because it depends on CPU cores available during execution
-	static_assert(MaxChunksPerBatch > 0 && (MaxChunksPerBatch & (MaxChunksPerBatch - 1)) == 0);
-	int NumBatches                                                     = 1;
-	int NumChunksPerBatch                                              = (int)MaxChunksPerBatch;
-	uint32 MaxLoopsPerChunk                                            = 0;
-	{ //compute the number of bytes required per batch
-		const uint32 TotalNumLoopsRequired              = VVM_MAX(((uint32)InitData->NumInstances + 3) >> 2, 1);
-		const size_t NumBytesRequiredPerLoop            = sizeof(FVecReg) * InitData->OptimizeContext->NumTempRegisters;
-		if (BatchOverheadSize + 64 > GVVMChunkSizeInBytes) {
-			//either the chunk size is way too small, or there's an insane number of consts or data sets required... we just revert to the previous VM's default
-			MaxLoopsPerChunk                            = 128 >> 2;
-			NumBatches                                  = (TotalNumLoopsRequired + MaxLoopsPerChunk - 1) / MaxLoopsPerChunk;
-		}
-		else
+	FVectorVMState *VVMState = (FVectorVMState *)StatePtr;
+	
+	{ //setup the pointers that are allocated in conjunction with the VVMState
+		uint32 ExtFnTableOffset       = (uint32)(ConstBufferOffset      + ConstantBufferNumBytes);
+		uint32 OutputPerDataSetOffset = (uint32)(ExtFnTableOffset       + ExtFnTableNumBytes);
+		uint32 ConstMapCacheOffset    = (uint32)(OutputPerDataSetOffset + OutputPerDataSetNumBytes);
+		uint32 InputMapCacheOffset    = (uint32)(ConstMapCacheOffset    + ConstMapCacheNumBytes);
+
+		VVMState->ConstantBuffers           = (FVecReg *)                   (StatePtr + ConstBufferOffset     );
+		VVMState->ExtFunctionTable          = (FVectorVMExtFunctionData *)  (StatePtr + ExtFnTableOffset      );
+		VVMState->NumOutputPerDataSet       = (volatile int32 *)            (StatePtr + OutputPerDataSetOffset);
+		VVMState->ConstMapCacheIdx          = (uint8 *)                     (StatePtr + ConstMapCacheOffset   );
+		VVMState->InputMapCacheIdx          = (uint8 *)                     (StatePtr + InputMapCacheOffset   );
+		VVMState->ConstMapCacheSrc          = (uint16 *)                    (VVMState->ConstMapCacheIdx + OptimizeCtx->NumConstsRemapped);
+		VVMState->InputMapCacheSrc          = (uint16 *)                    (VVMState->InputMapCacheIdx + OptimizeCtx->NumInputsRemapped);
+
+		check((size_t)((uint8 *)VVMState->InputMapCacheSrc - StatePtr) + sizeof(uint16) * OptimizeCtx->NumInputsRemapped <= VVMStateTotalNumBytes);
+
+		VVMState->NumInstancesExecCached    = 0;
+		for (uint32 i = 0; i < OptimizeCtx->NumExtFns; ++i)
 		{
-			size_t NumBytesPerBatchAvailableForTempRegs = GVVMChunkSizeInBytes - BatchOverheadSize;
-			size_t TotalNumLoopBytesRequired            = VVM_ALIGN(TotalNumLoopsRequired, MaxChunksPerBatch) * NumBytesRequiredPerLoop;
-			if (NumBytesPerBatchAvailableForTempRegs < TotalNumLoopBytesRequired)
-			{
-				//Not everything fits into a single chunk, so we have to compute everything here
-				int NumChunksRequired                   = (int)(TotalNumLoopBytesRequired + NumBytesPerBatchAvailableForTempRegs - 1) / (int)NumBytesPerBatchAvailableForTempRegs;
-				check(NumChunksRequired > 1);
-				if (NumChunksRequired < MaxChunksPerBatch) //everything fits in a single batch
-				{
-					NumChunksPerBatch                   = NumChunksRequired;
-					//take as little memory as possible and execute it in equal sized chunks
-					MaxLoopsPerChunk                    = (TotalNumLoopsRequired + NumChunksRequired - 1) / NumChunksRequired;
-				}
-				else //not everything fits in a single batch, we have to thread this
-				{
-					MaxLoopsPerChunk                    = (uint32)(NumBytesPerBatchAvailableForTempRegs / NumBytesRequiredPerLoop);
-					uint32 NumLoopsPerBatch             = MaxLoopsPerChunk * NumChunksPerBatch;
-					NumBatches                          = (TotalNumLoopsRequired + NumLoopsPerBatch - 1) / NumLoopsPerBatch;
-					if (GVVMMaxThreadsPerScript > 0 && NumBatches > GVVMMaxThreadsPerScript)
-					{
-						//number of batches exceed the number of threads allowed.. increase the number of chunks per batch
-						NumLoopsPerBatch                = (TotalNumLoopsRequired + GVVMMaxThreadsPerScript - 1) / GVVMMaxThreadsPerScript;
-						NumChunksPerBatch               = (int)(NumLoopsPerBatch + MaxLoopsPerChunk - 1) / MaxLoopsPerChunk;
-						NumBatches                      = GVVMMaxThreadsPerScript;
-						check(NumBatches * NumChunksPerBatch * MaxLoopsPerChunk >= TotalNumLoopsRequired);
-					}
-				}
-			}
-			else
-			{
-				//everything fits into a single chunk
-				NumChunksPerBatch = 1;
-				MaxLoopsPerChunk = TotalNumLoopsRequired;
-			}
+			VVMState->ExtFunctionTable[i].Function   = nullptr;
+			VVMState->ExtFunctionTable[i].NumInputs  = OptimizeCtx->ExtFnTable[i].NumInputs;
+			VVMState->ExtFunctionTable[i].NumOutputs = OptimizeCtx->ExtFnTable[i].NumOutputs;
 		}
-		PerBatchRegisterDataBytesRequired = ConstantBufferSize + MaxLoopsPerChunk * NumBytesRequiredPerLoop;
 	}
 
-	size_t NumBytesRequiredPerBatch = PerBatchRegisterDataBytesRequired + PerBatchChunkLocalDataOutputIdxBytesRequired + PerBatchChunkLocalNumOutputBytesRequired + PerBatchChunkLocalNumExtFnDecodeRegisterBytesRequired;
-	check(NumBytesRequiredPerBatch <= GVVMChunkSizeInBytes || BatchOverheadSize > GVVMChunkSizeInBytes); //if BatchOverheadSize is too high, then this check is invalid since we're blowing past the limit anyway and there's nothing we can do about it
-
-	{
-		//after the batch size has been calcluated we add two more things: padding and indices for deterministic random generation
-		//This will often cause the actual memory usage to go slightly above the GVVMChunkSizeInBytes.  I decided this is okay, because
-		//the concern is more to get the runtime memory usage lower than the L1 size, not to split hairs over a few hundred bytes.
-		//The random counters are very rarely needed, and are only included for backwards compatiblity with the previous VM.
-		//I figured it'd be better to append these at the end and not mess up what fits in the L1 for the 99.99% of cases that don't require it.
-		//There is no way to determine ahead of time if a script will need the random counters or not.
-		NumBytesRequiredPerBatch += VVM_CACHELINE_SIZE;
-		NumBytesRequiredPerBatch += sizeof(int32) * MaxLoopsPerChunk * 4;
+	{ //setup the pointers from the optimize context
+		VVMState->ConstRemapTable      = OptimizeCtx->ConstRemap[1];
+		VVMState->InputRemapTable      = OptimizeCtx->InputRemapTable;
+		VVMState->InputDataSetOffsets  = OptimizeCtx->InputDataSetOffsets;
+
+		VVMState->NumTempRegisters     = OptimizeCtx->NumTempRegisters;
+		VVMState->NumConstBuffers      = OptimizeCtx->NumConstsRemapped;
+		VVMState->NumInputBuffers      = OptimizeCtx->NumInputsRemapped;
+		VVMState->Bytecode             = OptimizeCtx->OutputBytecode;
+		VVMState->NumBytecodeBytes     = OptimizeCtx->NumBytecodeBytes;
+		VVMState->NumInputDataSets     = OptimizeCtx->NumInputDataSets;
+		VVMState->NumOutputDataSets    = OptimizeCtx->NumOutputDataSets;
+		VVMState->NumDummyRegsRequired = OptimizeCtx->NumDummyRegsReq;
+		VVMState->Flags                = OptimizeCtx->Flags & ~VVMFlag_DataMapCacheSetup;
+
+		VVMState->NumExtFunctions      = OptimizeCtx->NumExtFns;
+		VVMState->MaxExtFnRegisters    = OptimizeCtx->MaxExtFnRegisters;
 	}
 
-	{ //compute the number of overhead bytes for this VVM State
-		size_t NumDataSetOutputBytesRequired            = sizeof(volatile long)            * InitData->OptimizeContext->NumOutputDataSets;
-		size_t NumBatchStateBytesRequired               = sizeof(FVectorVMBatchState)      * NumBatches + 16; //this is *NOT* the memory per batch, just the overhead stored in VVMState
-		size_t NumExtFnBytesRequired                    = sizeof(FVectorVMExtFunctionData) * InitData->ExtFunctionTable.Num();
-		NumVVMStateBytesRequired                        = VVM_ALIGN_64(NumDataSetOutputBytesRequired +
-			                                                           NumExtFnBytesRequired +
-			                                                           NumBatchStateBytesRequired);
+#	if VECTORVM_SUPPORTS_AVX
+	{ //@TODO check CPUID stuff the Unreal way here
+		VVMState->Flags |= VVMFlag_SupportsAVX;
+	}
+#	endif
+
+	{ //compute fixed batch size
+		const size_t NumPtrRegsInTable = VVMState->NumTempRegisters + 
+		                                 VVMState->NumConstBuffers +
+		                                 VVMState->NumInputBuffers * 2;
+
+		const size_t ChunkLocalDataOutputIdxNumBytes     = sizeof(uint32) * VVMState->NumOutputDataSets;
+		const size_t ChunkLocalNumOutputNumBytes         = sizeof(uint32) * VVMState->NumOutputDataSets;
+		const size_t ChunkLocalNumExtFnDecodeRegNumBytes = (sizeof(FVecReg *) + sizeof(uint8)) * VVMState->MaxExtFnRegisters + VVM_REG_SIZE * VVMState->NumDummyRegsRequired;
+		const size_t RegPtrTableNumBytes                 = (sizeof(uint32 *) + sizeof(uint8)) * NumPtrRegsInTable;
+		const size_t BatchOverheadSize                   = VVM_ALIGN_64(sizeof(FVectorVMBatchState))              +
+			                                                            ChunkLocalDataOutputIdxNumBytes           +
+		                                                                ChunkLocalNumOutputNumBytes               +
+			                                                            ChunkLocalNumExtFnDecodeRegNumBytes       +
+			                                                            RegPtrTableNumBytes                       +
+			                                                            VVM_CHUNK_FIXED_OVERHEAD_SIZE;
+
+		VVMState->BatchOverheadSize                   = (uint32)BatchOverheadSize;
+		VVMState->ChunkLocalDataOutputIdxNumBytes     = (uint32)ChunkLocalDataOutputIdxNumBytes;
+		VVMState->ChunkLocalNumOutputNumBytes         = (uint32)ChunkLocalNumOutputNumBytes;
 	}
 
-	size_t TotalBytesRequired = VVM_ALIGN_64(sizeof(FVectorVMState) + NumVVMStateBytesRequired) + NumBytesRequiredPerBatch; //first batch gets allocated immediately following the VVM State
+	return VVMState;
+}
 
-	VectorVMReallocFn * ReallocFn = InitData->ReallocFn ? InitData->ReallocFn : VVMDefaultRealloc;
-	VectorVMFreeFn *    FreeFn    = InitData->FreeFn    ? InitData->FreeFn    : VVMDefaultFree;
-	
-	FVectorVMState *VVMState;
-	if (InitData->ExistingVectorVMState)
-	{
-		//start at 1 because the first batch is allocated immediately following the VVMState and will be freed when VVMState is freed
-		for (int i = 1; i < InitData->ExistingVectorVMState->NumBatches; ++i)
-		{
-			FreeFn(InitData->ExistingVectorVMState->BatchStates[i].MallocedMemPtr, __FILE__, __LINE__);
-		}
-	}
-	if (InitData->ExistingVectorVMState && InitData->ExistingVectorVMState->NumBytesMalloced >= TotalBytesRequired)
-	{
-		VVMState = InitData->ExistingVectorVMState;
+#include "VVMExec_SingleLoop.inl"
+#include "VVMExec_MultipleLoops.inl"
+#if VECTORVM_SUPPORTS_AVX
+#include "VVMExec_MultipleLoopsAVX.inl"
+#endif
+
+static void ExecVVMBatch(FVectorVMExecContext *ExecCtx, int ExecIdx, FVectorVMSerializeState *SerializeState, FVectorVMSerializeState *CmpSerializeState) {
+	if (ExecCtx->VVMState->Bytecode == nullptr) {
+		return;
 	}
-	else
+	VVMRAIIPageHandle PageHandle;
+	FVectorVMBatchState *BatchState = (FVectorVMBatchState *)VVMAllocBatch(ExecCtx->VVMState->BatchOverheadSize + ExecCtx->Internal.PerBatchRegisterDataBytesRequired, &PageHandle);
+
+	bool RegIncSetForAVX = false;
+	if (SetupBatchStatePtrs(ExecCtx, BatchState) == nullptr)
 	{
-		VVMState = (FVectorVMState *)ReallocFn(InitData->ExistingVectorVMState, TotalBytesRequired, __FILE__, __LINE__);
-		VVMState->NumBytesMalloced = TotalBytesRequired;
+		return;
 	}
-	uint8 *VVMStatePtr       = (uint8 *)VVM_ALIGN_16((size_t)(VVMState + 1)); //start state pointers immediately following the VVMState
-	if (VVMState == nullptr)
+	if (ExecCtx->VVMState->Flags & VVMFlag_HasRandInstruction)
 	{
-		return nullptr;
-	}
-	
-	VVMState->ReallocFn             = ReallocFn;
-	VVMState->FreeFn                = FreeFn;
-	VVMState->Error.Flags           = 0;
-	VVMState->Error.LineNum         = 0;
-	
-#	define IncVVMStatePtr(NumBytes, ExtraErrFlags)                                      \
-	VVMStatePtr += NumBytes;                                                            \
-	if ((size_t)((uint8 *)VVMStatePtr - (uint8 *)VVMState) > TotalBytesRequired)        \
-	{                                                                                   \
-		VVMState->Error.Flags |= VVMErr_InitMemMismatch | VVMErr_Fatal | ExtraErrFlags; \
-		VVMState->Error.LineNum = __LINE__;                                             \
-		return VVMState;                                                                \
+		SetupRandStateForBatch(BatchState);
 	}
 
-	{ //we have enough memory malloced for this state, so set up the stuff that comes immediately after the VVMState
-		if (InitData->OptimizeContext->NumExtFns > 0)
-		{
-			VVMState->NumExtFunctions  = InitData->OptimizeContext->NumExtFns;
-			VVMState->ExtFunctionTable = (FVectorVMExtFunctionData *)VVMStatePtr;        IncVVMStatePtr(sizeof(FVectorVMExtFunctionData) * VVMState->NumExtFunctions, 0);
-			for (uint32 i = 0; i < InitData->OptimizeContext->NumExtFns; ++i)
-			{
-				VVMState->ExtFunctionTable[i].Function   = InitData->ExtFunctionTable[i];
-				VVMState->ExtFunctionTable[i].NumInputs  = InitData->OptimizeContext->ExtFnTable[i].NumInputs;
-				VVMState->ExtFunctionTable[i].NumOutputs = InitData->OptimizeContext->ExtFnTable[i].NumOutputs;
-			}
-		}
-		else
+	//Loop until all available batches have been completed... BatchCounter will likely be very small, and if AssignInstancesToBatch return null
+	//then we return from the function, so just loop a lot!  This prevents the thread from exiting only to cause another thread to re-create
+	//batch state.  We can re-use everything as many times as necessary in times of high thread contention.
+	for (int BatchCounter = 0; BatchCounter < 1024; ++BatchCounter) {
+		if (!AssignInstancesToBatch(ExecCtx, BatchState))
 		{
-			VVMState->ExtFunctionTable = nullptr;
-			VVMState->NumExtFunctions  = 0;
+			return; //no more instances to do, we're done
 		}
-		uint8 *StateAfterExtFnPtr = VVMStatePtr;
-		VVMState->NumOutputPerDataSet  = (volatile int32 *)VVMStatePtr;                  IncVVMStatePtr(sizeof(volatile int32) * InitData->OptimizeContext->NumOutputDataSets, 0);
-		VVMStatePtr = (uint8 *)VVM_ALIGN_16(VVMStatePtr);
-		VVMState->BatchStates          = (FVectorVMBatchState *)VVMStatePtr;             IncVVMStatePtr(sizeof(volatile FVectorVMBatchState) * NumBatches, 0);
-		FMemory::Memset(StateAfterExtFnPtr, 0, (size_t)(VVMStatePtr - StateAfterExtFnPtr));
-	}
-#	undef IncVVMStatePtr
-
-	//init the part of the VVMState that doesn't require the externally allocated memory
-	VVMState->Bytecode          = InitData->OptimizeContext->OutputBytecode;
-	VVMState->NumBytecodeBytes  = InitData->OptimizeContext->NumBytecodeBytes;
-	VVMState->NumTempRegisters  = InitData->OptimizeContext->NumTempRegisters;
-	VVMState->NumConstBuffers   = InitData->OptimizeContext->NumConstsRemapped;
-	VVMState->NumOutputDataSets = InitData->OptimizeContext->NumOutputDataSets;
-	VVMState->MaxExtFnRegisters = MaxExtFnRegisters;
-	VVMState->NumDummyRegsReq   = InitData->OptimizeContext->NumDummyRegsReq;
-	VVMState->DataSets          = InitData->DataSets;
-
-	VVMState->UserPtrTable      = InitData->UserPtrTable;
-	VVMState->NumUserPtrTable   = InitData->NumUserPtrTable;
-	VVMState->TotalNumInstances = InitData->NumInstances;
-
-	//batch data
-	VVMState->NumBytesRequiredPerBatch                              = NumBytesRequiredPerBatch;
-	VVMState->PerBatchRegisterDataBytesRequired                     = PerBatchRegisterDataBytesRequired;
-	VVMState->PerBatchChunkLocalDataOutputIdxBytesRequired          = PerBatchChunkLocalDataOutputIdxBytesRequired;
-	VVMState->PerBatchChunkLocalNumOutputBytesRequired              = PerBatchChunkLocalNumOutputBytesRequired;
-	VVMState->PerBatchChunkLocalNumExtFnDecodeRegisterBytesRequired = PerBatchChunkLocalNumExtFnDecodeRegisterBytesRequired;
-
-	VVMState->NumBatches                     = NumBatches;
-	VVMState->MaxChunksPerBatch              = NumChunksPerBatch;
-	VVMState->MaxInstancesPerChunk           = MaxLoopsPerChunk << 2;
-	VVMState->NumInstancesAssignedToBatches  = 0;
-	VVMState->NumInstancesCompleted          = 0;
-
-	if (NumBatches > 0) //init the first batch using the remainder of the memory
-	{
-		FVectorVMBatchState *BatchState = VVMState->BatchStates + 0;
-		VVMStatePtr = SetupBatchStatePtrs(VVMState, BatchState, VVMStatePtr);
-
-		if (VVMStatePtr) //init the constant data at the start of the batch state
+		int StartInstanceThisChunk = BatchState->StartInstance;
+		int NumChunksThisBatch     = (BatchState->NumInstances + ExecCtx->Internal.MaxInstancesPerChunk - 1) / ExecCtx->Internal.MaxInstancesPerChunk;
+		for (int ChunkIdxThisBatch = 0; ChunkIdxThisBatch < NumChunksThisBatch; ++ChunkIdxThisBatch, StartInstanceThisChunk += ExecCtx->Internal.MaxInstancesPerChunk)
 		{
-			//first compute the starting offset of each set of constants
-			uint32 ConstCountAcc = 0;
-			for (int i = 0; i < InitData->NumConstData; ++i)
+			int NumInstancesThisChunk = VVM_MIN((int)ExecCtx->Internal.MaxInstancesPerChunk, BatchState->StartInstance + BatchState->NumInstances - StartInstanceThisChunk);
+			int NumLoops               = (int)((NumInstancesThisChunk + 3) & ~3) >> 2; //assumes 4-wide ops
+			
+			if (NumLoops == 1)
 			{
-				InitData->ConstData[i].StartingOffset = ConstCountAcc;
-				ConstCountAcc += InitData->ConstData[i].NumDWords;
+#				if VECTORVM_SUPPORTS_AVX
+				if (RegIncSetForAVX) {
+					SetRegIncTableFor4WideFrom8Wide(ExecCtx, BatchState);
+					RegIncSetForAVX = false;
+				}
+#				endif
+				execChunkSingleLoop(ExecCtx, BatchState, StartInstanceThisChunk, NumInstancesThisChunk, SerializeState, CmpSerializeState);
 			}
-			
-			FVecReg *ConstantBuffers = BatchState->RegisterData;
-			for (uint32 i = 0; i < VVMState->NumConstBuffers; ++i)
+			else if (NumLoops >= 1)
 			{
-				uint16 ConstBufferOffset = InitData->OptimizeContext->ConstRemap[1][i];
-				for (int j = 0; j < InitData->NumConstData; ++j)
+#				if VECTORVM_SUPPORTS_AVX
+				if (NumLoops >= 4 && (ExecCtx->VVMState->Flags & VVMFlag_SupportsAVX))
+				{
+					//SetupBatchPointers sets the register increment table to 16 bytes, with AVX we want 32 bytes.
+					if (!RegIncSetForAVX)
+					{
+						SetRegIncTableFor8WideFrom4Wide(ExecCtx, BatchState);
+						RegIncSetForAVX = true;
+					}
+					int NumAVXLoops = NumLoops / 2;
+					execChunkMultipleLoopsAVX(ExecCtx, BatchState, StartInstanceThisChunk, NumInstancesThisChunk, NumAVXLoops, SerializeState, CmpSerializeState);
+					if (NumAVXLoops * 2 != NumLoops)
+					{
+						check(false); //does this even work!? need to run it to find out
+						//we have an extra single loop to do at the end
+						NumLoops = 1;
+						SetRegIncTableFor4WideFrom8Wide(ExecCtx, BatchState);
+						RegIncSetForAVX = false;
+						execChunkSingleLoop(ExecCtx, BatchState, NumAVXLoops * 8, NumInstancesThisChunk - NumAVXLoops * 8, SerializeState, CmpSerializeState);
+					}
+				}
+				else 
 				{
-					FVectorVMConstData *ConstData = InitData->ConstData + j;
-					if (ConstBufferOffset >= ConstData->StartingOffset && ConstBufferOffset < ConstData->StartingOffset + ConstData->NumDWords)
+					if (RegIncSetForAVX)
 					{
-						ConstantBuffers[i].i = VectorIntSet1(((uint32 *)ConstData->RegisterData)[ConstBufferOffset - ConstData->StartingOffset]);
-						break;
+						SetRegIncTableFor4WideFrom8Wide(ExecCtx, BatchState);
+						RegIncSetForAVX = false;
 					}
+#				else
+				{
+#				endif
+					execChunkMultipleLoops(ExecCtx, BatchState, StartInstanceThisChunk, NumInstancesThisChunk, NumLoops, SerializeState, CmpSerializeState);
 				}
 			}
 		}
-		else
+	
+		VVMSer_batchEndExp(SerializeState);
+
+		if (BatchState->ChunkLocalData.RandCounters)
 		{
-			return VVMState;
+			FMemory::Free(BatchState->ChunkLocalData.RandCounters);
+			BatchState->ChunkLocalData.RandCounters = NULL;
 		}
-		SetupRandStateForBatch(BatchState);
-	}
-	if (SerializeState)
-	{
-		VVMSer_initSerializationState(VVMState, SerializeState, InitData, SerializeState->Flags | VVMSer_OptimizedBytecode);
+
+		FPlatformAtomics::InterlockedAdd(&ExecCtx->Internal.NumInstancesCompleted, BatchState->NumInstances);
 	}
-	return VVMState;
 }
 
-
-#	define execVecIns1f(ins_)						VVMSer_instruction(0, 1)                                            \
-													for (int i = 0; i < NumLoops; ++i)                                  \
-													{                                                                   \
-														VectorRegister4f r0  = VectorLoad(&VecReg[0][i & RegInc[0]].v); \
-														VectorRegister4f res = ins_(r0);                                \
-														VectorStoreAligned(res, &VecReg[1][i].v);                       \
-													}                                                                   \
-													InsPtr += 4;
-
-#	define execVecIns2f(ins_)						VVMSer_instruction(0, 2)                                            \
-													for (int i = 0; i < NumLoops; ++i)                                  \
-													{                                                                   \
-														VectorRegister4f r0  = VectorLoad(&VecReg[0][i & RegInc[0]].v); \
-														VectorRegister4f r1  = VectorLoad(&VecReg[1][i & RegInc[1]].v); \
-														VectorRegister4f res = ins_(r0, r1);                            \
-														VectorStoreAligned(res, &VecReg[2][i].v);                       \
-													}                                                                   \
-													InsPtr += 6;
-
-#	define execVecIns3f(ins_)						VVMSer_instruction(0, 3)                                            \
-													for (int i = 0; i < NumLoops; ++i)                                  \
-													{                                                                   \
-														VectorRegister4f r0  = VectorLoad(&VecReg[0][i & RegInc[0]].v); \
-														VectorRegister4f r1  = VectorLoad(&VecReg[1][i & RegInc[1]].v); \
-														VectorRegister4f r2  = VectorLoad(&VecReg[2][i & RegInc[2]].v); \
-														VectorRegister4f res = ins_(r0, r1, r2);                        \
-														VectorStoreAligned(res, &VecReg[3][i].v);                       \
-													}                                                                   \
-													InsPtr += 8;
-
-#	define execVecIns1i(ins_)						VVMSer_instruction(1, 1)                                               \
-													for (int i = 0; i < NumLoops; ++i)                                     \
-													{                                                                      \
-														VectorRegister4i r0  = VectorIntLoad(&VecReg[0][i & RegInc[0]].v); \
-														VectorRegister4i res = ins_(r0);                                   \
-														VectorIntStoreAligned(res, &VecReg[1][i].i);                       \
-													}                                                                      \
-													InsPtr += 4;
-
-#	define execVecIns2i(ins_)						VVMSer_instruction(1, 2)                                               \
-													for (int i = 0; i < NumLoops; ++i)                                     \
-													{                                                                      \
-														VectorRegister4i r0  = VectorIntLoad(&VecReg[0][i & RegInc[0]].v); \
-														VectorRegister4i r1  = VectorIntLoad(&VecReg[1][i & RegInc[1]].v); \
-														VectorRegister4i res = ins_(r0, r1);                               \
-														VectorIntStoreAligned(res, &VecReg[2][i].i);                       \
-													}                                                                      \
-													InsPtr += 6;
-
-#	define execVecIns3i(ins_)						VVMSer_instruction(1, 3)                                               \
-													for (int i = 0; i < NumLoops; ++i)                                     \
-													{                                                                      \
-														VectorRegister4i r0  = VectorIntLoad(&VecReg[0][i & RegInc[0]].v); \
-														VectorRegister4i r1  = VectorIntLoad(&VecReg[1][i & RegInc[1]].v); \
-														VectorRegister4i r2  = VectorIntLoad(&VecReg[2][i & RegInc[2]].v); \
-														VectorRegister4i res = ins_(r0, r1, r2);                           \
-														VectorIntStoreAligned(res, &VecReg[3][i].i);                       \
-													}                                                                      \
-													InsPtr += 8;
-
-#		define VVMDecodeInstructionRegisters(Bytecode, RegData, RegIncMask) do {                                                                                                                                                            \
-			uint32 VecOffsets[4];                                                                                                                                                                                                           \
-			VectorRegister4i VecIndicesIn4  = VectorIntLoad(Bytecode);                                                                          /* 16 bit inputs.  15 bits for index, 1 high bit for const/reg flag (0: reg, 1: const) */   \
-			VectorRegister4i VecIndices4    = VectorIntExpandLow16To32(VecIndicesIn4);                                                          /* 4-wide 32 bit version of the inputs, bits 16:31 are 0 */                                 \
-			VectorRegister4i RegIncInv4     = VectorIntSubtract(VectorSetZero(), VectorShiftRightImmArithmetic(VecIndices4, 15));               /* Sets the inverse of what we need: 0xFF... is const, 0 is reg */                          \
-			VectorRegister4i RegInc4        = VectorIntXor(RegIncInv4, VVM_m128iConst(FMask));                                                  /* Whether to increment the index counter: 0xFF... for registers, and 0 for const */        \
-			VectorRegister4i VecRegIndices4 = VectorIntAnd(VecIndices4, VVM_m128iConst(RegOffsetMask));                                         /* Only the register index, the const/reg flag is stripped off */                           \
-			VectorRegister4i ConstOffset4   = VectorIntAnd(VecRegIndices4, RegIncInv4);                                                         /* Only the const offsets, all registers are masked out */                                  \
-			VectorRegister4i TempRegOffset4 = VectorIntAnd(VectorIntAdd(NumConsts4, VectorIntMultiply(VecRegIndices4, NumLoops4)), RegInc4);    /* Only the register offsets, all consts are masked out */                                  \
-			VectorRegister4i OptRegOffset4  = VectorIntOr(ConstOffset4, TempRegOffset4);                                                        /* Blended (sse4 would be nice) const and temp register offsets */                          \
-			VectorIntStore(RegInc4, RegIncMask);                                                                                                                                                                                            \
-			VectorIntStore(OptRegOffset4, VecOffsets);                                                                                                                                                                                      \
-			(RegData)[0] = BatchState->RegisterData + VecOffsets[0];                                                                                                                                                                        \
-			(RegData)[1] = BatchState->RegisterData + VecOffsets[1];                                                                                                                                                                        \
-			(RegData)[2] = BatchState->RegisterData + VecOffsets[2];                                                                                                                                                                        \
-			(RegData)[3] = BatchState->RegisterData + VecOffsets[3];                                                                                                                                                                        \
-		} while (0)
-
-static void ExecVVMBatch(FVectorVMState *VVMState, int ExecIdx, FVectorVMSerializeState *SerializeState, FVectorVMSerializeState *CmpSerializeState)
+static void VVMBuildMapTableCaches(FVectorVMExecContext *ExecCtx)
 {
-	FVectorVMBatchState *BatchState = nullptr;
-	int BatchIdx = -1;
-	//check to see if we can reuse a batch's memory that's finished executing
-	for (int i = 0; i < VVMState->NumBatches; ++i)
-	{
-		int32 WasCurrentlyExecuting = FPlatformAtomics::InterlockedCompareExchange(&VVMState->BatchStates[i].CurrentlyExecuting, 1, 0);
-		if (WasCurrentlyExecuting == 0)
-		{
-			//we can reuse this batch
-			BatchState = VVMState->BatchStates + i;
-			BatchIdx = i;
-			break;
-		}
-	}
-	if (BatchIdx == -1 || BatchState == nullptr)
+	//constant buffers
+	check(ExecCtx->VVMState->NumConstBuffers < 0xFF);
+	for (uint32 i = 0; i < ExecCtx->VVMState->NumConstBuffers; ++i)
 	{
-		check(false);
-		return;
-	}
-	check((((size_t)BatchState) & 0xF) == 0); //must be 16 byte aligned
-	if (BatchState->RegisterData == nullptr)
-	{
-		check(BatchIdx != 0); //first batch state should have set the pointers in init()
-		//this is the first time using this batch, so we need to malloc the data and copy the consts over from batch state 0
-		check(VVMState->NumBytesRequiredPerBatch >= VVMState->PerBatchRegisterDataBytesRequired +
-		                                            VVMState->PerBatchChunkLocalDataOutputIdxBytesRequired +
-		                                            VVMState->PerBatchChunkLocalNumOutputBytesRequired +
-			                                        VVM_CACHELINE_SIZE /* padding */);
-		void *BatchData = VVMState->ReallocFn(nullptr, VVMState->NumBytesRequiredPerBatch, __FILE__, __LINE__);
-		if (BatchData == nullptr)
-		{
-			check(false);
-			VVMState->Error.Flags |= VVMErr_BatchMemory;
-			VVMState->Error.LineNum = __LINE__;
-			return;
-		}
-		if (SetupBatchStatePtrs(VVMState, BatchState, BatchData) == nullptr)
+		uint32 RemappedIdx = ExecCtx->VVMState->ConstRemapTable[i];
+		uint32 ConstCountAcc = 0;
+		for (int j = 0; j < ExecCtx->ConstantTableCount; ++j)
 		{
-			return;
+			const uint32 NumDWords = (uint32)ExecCtx->ConstantTableSizes[j] >> 2;
+			if (ConstCountAcc + NumDWords > RemappedIdx)
+			{
+				const uint32 *SrcConstArray = (uint32 *)ExecCtx->ConstantTableData[j];
+				uint32 Idx = RemappedIdx - ConstCountAcc;
+				check(Idx < 0xFFFF);
+				ExecCtx->VVMState->ConstMapCacheIdx[i] = j;
+				ExecCtx->VVMState->ConstMapCacheSrc[i] = (uint16)(Idx);
+				break;
+			}
+			ConstCountAcc += NumDWords;
 		}
-		VVMMemCpy(BatchState->RegisterData, VVMState->BatchStates[0].RegisterData, sizeof(FVecReg) * VVMState->NumConstBuffers); //copy the constant data from the first batch
-		SetupRandStateForBatch(BatchState);
-	}
-	if (!AssignInstancesToBatch(VVMState, BatchState))
-	{
-		return; //no more instances to do, we're done
 	}
 
-	int StartInstanceThisChunk = BatchState->StartInstance;
-	int NumChunksThisBatch     = (BatchState->NumInstances + VVMState->MaxInstancesPerChunk - 1) / VVMState->MaxInstancesPerChunk;
-	
-	uint32 RegInc[4];
-	VectorRegister4i NumConsts4 = VectorIntSet1(VVMState->NumConstBuffers);
-	
-	for (int ChunkIdxThisBatch = 0; ChunkIdxThisBatch < NumChunksThisBatch; ++ChunkIdxThisBatch, StartInstanceThisChunk += VVMState->MaxInstancesPerChunk)
+	//inputs
+	if (ExecCtx->VVMState->NumInputDataSets > 0)
 	{
-		int NumInstancesThisChunk = VVM_MIN(VVMState->MaxInstancesPerChunk, BatchState->StartInstance + BatchState->NumInstances - StartInstanceThisChunk);
-		VVMSer_chunkStartExp(SerializeState, ChunkIdxThisBatch, BatchIdx);
-		
-		int NumLoops               = (int)((NumInstancesThisChunk + 3) & ~3) >> 2; //assumes 4-wide ops
-		VectorRegister4i NumLoops4 = VectorIntSet1(NumLoops);
-		uint8 *InsPtr              = VVMState->Bytecode;
-		uint8 *InsPtrEnd           = InsPtr + VVMState->NumBytecodeBytes;
-
-		for (uint32 i = 0; i < VVMState->NumOutputDataSets; ++i)
-		{
-			BatchState->ChunkLocalData.StartingOutputIdxPerDataSet[i] = 0;
-			BatchState->ChunkLocalData.NumOutputPerDataSet[i] = 0;
-		}
-		
-		FVecReg *VecReg[4];
-		
-		while (InsPtr < InsPtrEnd)
+		int InputCounter     = 0;
+		int NumInputDataSets = VVM_MIN(ExecCtx->VVMState->NumInputDataSets, (uint32)ExecCtx->DataSets.Num()); //Niagara can pass in any amount of datasets, but we only care about the highest one actually used as determined by the optimizer
+		for (int i = 0; i < NumInputDataSets; ++i)
 		{
-			VVMSer_insStartExp(SerializeState);
-			EVectorVMOp OpCode = (EVectorVMOp)*InsPtr++;
-			uint16 *VecIndices = (uint16 *)InsPtr;
-			VVMDecodeInstructionRegisters(InsPtr, VecReg, RegInc);
-			VVMSer_insEndDecodeExp(SerializeState);
-
-OpCodeSwitch: //I think computed gotos would be a huge win here... maybe write this loop in assembly for the jump table?!
-			switch (OpCode) {
-				case EVectorVMOp::done:                                                                 break;
-				case EVectorVMOp::add:                              execVecIns2f(VectorAdd);            break;
-				case EVectorVMOp::sub:                              execVecIns2f(VectorSubtract);       break;
-				case EVectorVMOp::mul:                              execVecIns2f(VectorMultiply);       break;
-				case EVectorVMOp::div:                              execVecIns2f(VVM_safeIns_div);      break;
-				case EVectorVMOp::mad:                              execVecIns3f(VectorMultiplyAdd);    break;
-				case EVectorVMOp::lerp:                             execVecIns3f(VectorLerp);           break;
-				case EVectorVMOp::rcp:                              execVecIns1f(VVM_safeIns_rcp);      break;
-				case EVectorVMOp::rsq:                              execVecIns1f(VVM_safeins_rsq);      break;
-				case EVectorVMOp::sqrt:                             execVecIns1f(VVM_safeins_sqrt);     break;
-				case EVectorVMOp::neg:                              execVecIns1f(VectorNegate);         break;
-				case EVectorVMOp::abs:                              execVecIns1f(VectorAbs);            break;
-				case EVectorVMOp::exp:                              execVecIns1f(VectorExp);            break;
-				case EVectorVMOp::exp2:                             execVecIns1f(VectorExp2);           break;
-				case EVectorVMOp::log:                              execVecIns1f(VVM_safeins_log);      break;
-				case EVectorVMOp::log2:                             execVecIns1f(VectorLog2);           break;
-				case EVectorVMOp::sin:                              execVecIns1f(VectorSin);            break;
-				case EVectorVMOp::cos:                              execVecIns1f(VectorCos);            break;
-				case EVectorVMOp::tan:                              execVecIns1f(VectorTan);            break;
-				case EVectorVMOp::asin:                             execVecIns1f(VectorASin);           break;
-				case EVectorVMOp::acos:                             execVecIns1f(VVM_vecACosFast);      break;
-				case EVectorVMOp::atan:                             execVecIns1f(VectorATan);           break;
-				case EVectorVMOp::atan2:                            execVecIns2f(VectorATan2);          break;
-				case EVectorVMOp::ceil:                             execVecIns1f(VectorCeil);           break;
-				case EVectorVMOp::floor:                            execVecIns1f(VectorFloor);          break;
-				case EVectorVMOp::fmod:                             execVecIns2f(VectorMod);            break;
-				case EVectorVMOp::frac:                             execVecIns1f(VectorFractional);     break;
-				case EVectorVMOp::trunc:                            execVecIns1f(VectorTruncate);       break;
-				case EVectorVMOp::clamp:                            execVecIns3f(VectorClamp);          break;
-				case EVectorVMOp::min:                              execVecIns2f(VectorMin);            break;
-				case EVectorVMOp::max:                              execVecIns2f(VectorMax);            break;
-				case EVectorVMOp::pow:                              execVecIns2f(VVM_safeins_pow);      break;
-				case EVectorVMOp::round:                            execVecIns1f(VectorRound);          break;
-				case EVectorVMOp::sign:                             execVecIns1f(VectorSign);           break;
-				case EVectorVMOp::step:                             execVecIns2f(VVM_vecStep);          break;
-				case EVectorVMOp::random:
-#					if defined(VVM_INCLUDE_SERIALIZATION) && !defined(VVM_SERIALIZE_NO_WRITE)
-					if (SerializeState && (SerializeState->Flags & VVMSer_SyncRandom) && CmpSerializeState && CmpSerializeState->NumInstructions >= SerializeState->NumInstructions && CmpSerializeState->NumTempRegisters > VecIndices[1])
-					{
-						FVectorVMSerializeInstruction *CmpIns = nullptr;
-						if ((SerializeState->Flags & VVMSer_OptimizedBytecode) && SerializeState->OptimizeCtx && !(CmpSerializeState->Flags & VVMSer_OptimizedBytecode))
-						{
-							for (uint32 i = 0; i < SerializeState->OptimizeCtx->Intermediate.NumInstructions; ++i)
-							{
-								if (SerializeState->OptimizeCtx->Intermediate.Instructions[i].PtrOffsetInOptimizedBytecode + 1 == (uint32)(InsPtr - VVMState->Bytecode))
-								{ //+1 because we already incremented InsPtr
-									for (uint32 j = 0; j < CmpSerializeState->NumInstructions; ++j)
-									{
-										if (CmpSerializeState->Instructions[j].OpStart == SerializeState->OptimizeCtx->Intermediate.Instructions[i].PtrOffsetInOrigBytecode)
-										{
-											CmpIns = CmpSerializeState->Instructions + j;
-											break;
-										}
-									}
-									break;
-								}
-							}
-						}
-						else
-						{
-							CmpIns = CmpSerializeState->Instructions + SerializeState->NumInstructions + 1; //this instruction
-						}
-						if (CmpIns)
-						{
-							check(CmpSerializeState->Bytecode[CmpIns->OpStart] == (uint8)EVectorVMOp::random);
-							uint16 CmpVecIdx = ((uint16 *)(CmpSerializeState->Bytecode + CmpIns->OpStart + 2))[1];
-							check((CmpVecIdx & 8000) == 0); //can't output to constant
-							float *op_reg = (float *)VecReg[1];
-							float *ip_reg = (float *)CmpIns->TempRegisters + SerializeState->NumInstances * CmpVecIdx + StartInstanceThisChunk;
-							for (int i = 0; i < NumInstancesThisChunk; ++i)
-							{
-								*op_reg++ = *ip_reg++;
-							}
-							goto VVM_EndSyncRandom;
-						}
-					}
-#					endif //VVM_INCLUDE_SERIALIZATION
-					for (int i = 0; i < NumLoops; ++i)
-					{
-						VectorRegister4i RandReg = VVMXorwowStep(BatchState);
-						VectorRegister4i IntPart = VectorIntOr(VectorShiftRightImmLogical(RandReg, 9), VectorIntSet1(0x3F800000));
-						VectorRegister4f FltPart = VectorRegister4f(VectorCastIntToFloat(IntPart));
-						VecReg[1][i].v = VectorMultiply(VectorSubtract(FltPart, VectorSet1(1.f)), VecReg[0][i & RegInc[0]].v);
-					}
-#					if defined(VVM_INCLUDE_SERIALIZATION) && !defined(VVM_SERIALIZE_NO_WRITE)
-					VVM_EndSyncRandom:
-#					endif //VVM_INCLUDE_SERIALIZATION
-					VVMSer_regUsed(VecIndices[1], 0);
-					InsPtr += 4;
-				break;
-				case EVectorVMOp::noise:                            check(false);                       break;
-				case EVectorVMOp::cmplt:                            execVecIns2f(VectorCompareLT);      break;
-				case EVectorVMOp::cmple:                            execVecIns2f(VectorCompareLE);      break;
-				case EVectorVMOp::cmpgt:                            execVecIns2f(VectorCompareGT);      break;
-				case EVectorVMOp::cmpge:                            execVecIns2f(VectorCompareGE);      break;
-				case EVectorVMOp::cmpeq:                            execVecIns2f(VectorCompareEQ);      break;
-				case EVectorVMOp::cmpneq:                           execVecIns2f(VectorCompareNE);      break;
-				case EVectorVMOp::select:                           execVecIns3i(VectorIntSelect);      break;
-				case EVectorVMOp::addi:                             execVecIns2i(VectorIntAdd);         break;
-				case EVectorVMOp::subi:                             execVecIns2i(VectorIntSubtract);    break;
-				case EVectorVMOp::muli:                             execVecIns2i(VectorIntMultiply);    break;
-				case EVectorVMOp::divi: {
-					VVMSer_instruction(1, 2)
-					//@TODO: convert to double and div 4 wide
-					for (int i = 0; i < NumLoops; ++i)
-					{
-						int32 TmpA[4];
-						VectorIntStore(VecReg[0][i & RegInc[0]].i, TmpA);
-						int32 TmpB[4];
-						VectorIntStore(VecReg[1][i & RegInc[1]].i, TmpB);
-
-						// No intrinsics exist for integer divide. Since div by zero causes crashes, we must be safe against that.
-						int32 TmpDst[4];
-						TmpDst[0] = TmpB[0] != 0 ? (TmpA[0] / TmpB[0]) : 0;
-						TmpDst[1] = TmpB[1] != 0 ? (TmpA[1] / TmpB[1]) : 0;
-						TmpDst[2] = TmpB[2] != 0 ? (TmpA[2] / TmpB[2]) : 0;
-						TmpDst[3] = TmpB[3] != 0 ? (TmpA[3] / TmpB[3]) : 0;
-						VecReg[2][i].i = MakeVectorRegisterInt(TmpDst[0], TmpDst[1], TmpDst[2], TmpDst[3]);
-					}
-					InsPtr += 6;
-				} break;
-				case EVectorVMOp::clampi:                           execVecIns3i(VectorIntClamp);       break;
-				case EVectorVMOp::mini:                             execVecIns2i(VectorIntMin);         break;
-				case EVectorVMOp::maxi:                             execVecIns2i(VectorIntMax);         break;
-				case EVectorVMOp::absi:                             execVecIns1i(VectorIntAbs);         break;
-				case EVectorVMOp::negi:                             execVecIns1i(VectorIntNegate);      break;
-				case EVectorVMOp::signi:                            execVecIns1i(VectorIntSign);        break;
-				case EVectorVMOp::randomi: {
-					//@TODO: serialize syncing, no test cases yet
-					VVMSer_instruction(1, 1)
-					for (int i = 0; i < NumLoops; ++i)
-					{
-						VecReg[0][i].i = VVMXorwowStep(BatchState);
-					}
-					InsPtr += 4;
-				} break;
-				case EVectorVMOp::cmplti:                           execVecIns2i(VectorIntCompareLT);   break;
-				case EVectorVMOp::cmplei:                           execVecIns2i(VectorIntCompareLE);   break;
-				case EVectorVMOp::cmpgti:                           execVecIns2i(VectorIntCompareGT);   break;
-				case EVectorVMOp::cmpgei:                           execVecIns2i(VectorIntCompareGE);   break;
-				case EVectorVMOp::cmpeqi:                           execVecIns2i(VectorIntCompareEQ);   break;
-				case EVectorVMOp::cmpneqi:                          execVecIns2i(VectorIntCompareNEQ);  break;
-				case EVectorVMOp::bit_and:                          execVecIns2i(VectorIntAnd);         break;
-				case EVectorVMOp::bit_or:                           execVecIns2i(VectorIntOr);          break;
-				case EVectorVMOp::bit_xor:                          execVecIns2i(VectorIntXor);         break;
-				case EVectorVMOp::bit_not:                          execVecIns1i(VectorIntNot);         break;
-				case EVectorVMOp::bit_lshift:
-					for (int i = 0; i < NumLoops; ++i)
-					{
-						int idx0 = (i << 2) & RegInc[0];
-						int idx1 = (i << 2) & RegInc[1];
+			uint32 **DataSetInputBuffers = (uint32 **)ExecCtx->DataSets[i].InputRegisters.GetData();
+			int32 InstanceOffset = ExecCtx->DataSets[i].InstanceOffset;
 
-						((int *)VecReg[2])[(i << 2) + 0] = ((int *)VecReg[0])[idx0 + 0] << ((int *)VecReg[1])[idx1 + 0];
-						((int *)VecReg[2])[(i << 2) + 1] = ((int *)VecReg[0])[idx0 + 1] << ((int *)VecReg[1])[idx1 + 1];
-						((int *)VecReg[2])[(i << 2) + 2] = ((int *)VecReg[0])[idx0 + 2] << ((int *)VecReg[1])[idx1 + 2];
-						((int *)VecReg[2])[(i << 2) + 3] = ((int *)VecReg[0])[idx0 + 3] << ((int *)VecReg[1])[idx1 + 3];
-					}
-					VVMSer_regUsed(VecIndices[2], 1);
-					InsPtr += 6;
-					break;
-				case EVectorVMOp::bit_rshift:
-					for (int i = 0; i < NumLoops; ++i)
-					{
-						int idx0 = (i << 2) & RegInc[0];
-						int idx1 = (i << 2) & RegInc[1];
-						((int *)VecReg[2])[(i << 2) + 0] = ((int *)VecReg[0])[idx0 + 0] >> ((int *)VecReg[1])[idx1 + 0];
-						((int *)VecReg[2])[(i << 2) + 1] = ((int *)VecReg[0])[idx0 + 1] >> ((int *)VecReg[1])[idx1 + 1];
-						((int *)VecReg[2])[(i << 2) + 2] = ((int *)VecReg[0])[idx0 + 2] >> ((int *)VecReg[1])[idx1 + 2];
-						((int *)VecReg[2])[(i << 2) + 3] = ((int *)VecReg[0])[idx0 + 3] >> ((int *)VecReg[1])[idx1 + 3];
-					}
-					VVMSer_regUsed(VecIndices[2], 1);
-					InsPtr += 6;
-					break;
-				case EVectorVMOp::logic_and:                        execVecIns2i(VectorIntAnd);         break;
-				case EVectorVMOp::logic_or:                         execVecIns2i(VectorIntOr);          break;
-				case EVectorVMOp::logic_xor:                        execVecIns2i(VectorIntXor);         break;
-				case EVectorVMOp::logic_not:                        execVecIns1i(VectorIntNot);         break;
-				case EVectorVMOp::f2i:
-					for (int i = 0; i < NumLoops; ++i)
-					{
-						VectorRegister4f r0 = VectorLoad(&VecReg[0][i & RegInc[0]].v);
-						VectorIntStoreAligned(VectorFloatToInt(r0), &VecReg[1][i].i);
-					}
-					VVMSer_regUsed(VecIndices[1], 1);
-					InsPtr += 4;
-				break;
-				case EVectorVMOp::i2f:
-					for (int i = 0; i < NumLoops; ++i)
-					{
-						VectorRegister4i r0 = VectorIntLoad(&VecReg[0][i & RegInc[0]].i);
-						VectorStoreAligned(VectorIntToFloat(r0), &VecReg[1][i].v);
-					}
-					VVMSer_regUsed(VecIndices[1], 0);
-					InsPtr += 4;
-				break;
-				case EVectorVMOp::f2b:                              execVecIns1f(VVM_vecFloatToBool);   break;
-				case EVectorVMOp::b2f:                              execVecIns1f(VVM_vecBoolToFloat);   break;
-				case EVectorVMOp::i2b:                              execVecIns1i(VVM_vecIntToBool);     break;
-				case EVectorVMOp::b2i:                              execVecIns1i(VVM_vecBoolToInt);     break;
-				case EVectorVMOp::inputdata_float:
-				case EVectorVMOp::inputdata_int32: {
-					uint8 RegType             = (uint8)OpCode - (uint8)EVectorVMOp::inputdata_float;  //0: float, 1: int, 2: half
-					uint16 DataSetIdx         = *(uint16 *)(InsPtr    );
-					uint16 InputRegIdx        = *(uint16 *)(InsPtr + 2);
-					uint16 DestRegIdx         = *(uint16 *)(InsPtr + 4);
-					uint32 InstanceOffset     = VVMState->DataSets[DataSetIdx].InstanceOffset;
-					uint32 InputRegTypeOffset = VVMState->DataSets[DataSetIdx].InputRegisterTypeOffsets[RegType];
-					size_t DstIdx             = NumLoops * DestRegIdx;
-					uint32 **InputBuffers     = (uint32 **)VVMState->DataSets[DataSetIdx].InputRegisters.GetData();
-					VVMMemCpy(BatchState->RegisterData + VVMState->NumConstBuffers + DstIdx, InputBuffers[InputRegIdx + InputRegTypeOffset] + StartInstanceThisChunk + InstanceOffset, sizeof(FVecReg) * NumLoops);
-					VVMSer_regUsed(DestRegIdx, RegType);
-					InsPtr += 6;
-				} break;
-				case EVectorVMOp::inputdata_half: {
-					uint8 RegType             = 2;
-					uint16 DataSetIdx         = *(uint16 *)(InsPtr    );
-					uint16 InputRegIdx        = *(uint16 *)(InsPtr + 2);
-					uint16 DestRegIdx         = *(uint16 *)(InsPtr + 4);
-					uint32 InstanceOffset     = VVMState->DataSets[DataSetIdx].InstanceOffset;
-					uint32 InputRegTypeOffset = VVMState->DataSets[DataSetIdx].InputRegisterTypeOffsets[RegType];
-					size_t DstIdx             = NumLoops * DestRegIdx;
-					uint16 **InputBuffers     = (uint16 **)VVMState->DataSets[DataSetIdx].InputRegisters.GetData();
-					for (int i = 0; i < NumLoops; ++i) {
-						FPlatformMath::VectorLoadHalf((float *)(BatchState->RegisterData + VVMState->NumConstBuffers + DstIdx + i), (InputBuffers[InputRegIdx + InputRegTypeOffset] + StartInstanceThisChunk + InstanceOffset + i));
-					}
-					VVMSer_regUsed(DestRegIdx, 0);
-					InsPtr += 6;
-				} break;
-				case EVectorVMOp::inputdata_noadvance_float:
-				case EVectorVMOp::inputdata_noadvance_int32: {
-					uint8 RegType                 = (uint8)OpCode - (uint8)EVectorVMOp::inputdata_noadvance_float; //0: float, 1: int, 2: half
-					uint16 DataSetIdx             = *(uint16 *)(InsPtr    );
-					uint16 InputRegIdx            = *(uint16 *)(InsPtr + 2);
-					uint16 DestRegIdx             = *(uint16 *)(InsPtr + 4);
-					uint32 InstanceOffset         = VVMState->DataSets[DataSetIdx].InstanceOffset;
-					uint32 InputRegTypeOffset     = VVMState->DataSets[DataSetIdx].InputRegisterTypeOffsets[RegType];
-					size_t DstIdx                 = NumLoops * DestRegIdx;
-					uint32 **InputBuffers         = (uint32 **)VVMState->DataSets[DataSetIdx].InputRegisters.GetData();
-					uint32 *InputBuffer           = InputBuffers[InputRegIdx + InputRegTypeOffset] + InstanceOffset;
-					VectorRegister4i input_val4   = VectorIntSet1(*InputBuffer);
-					for (int i = 0; i < NumLoops; ++i)
-					{
-						VectorIntStoreAligned(input_val4, BatchState->RegisterData + VVMState->NumConstBuffers + DstIdx + i);
-					}
-					VVMSer_regUsed(DestRegIdx, RegType);
-					InsPtr += 6;
-				} break;
-				case EVectorVMOp::inputdata_noadvance_half:			check(false);				        break;
-				case EVectorVMOp::outputdata_float:
-				case EVectorVMOp::outputdata_int32: {
-					uint8 RegType = InsPtr[-1] - (uint8)EVectorVMOp::outputdata_float;
-					check(RegType == 0 || RegType == 1); //float or int32
-					uint16 DataSetIdx         = VecIndices[0];
-					int *DstIdxReg            = (int *)VecReg[1];
-					uint32 NumOutputInstances = BatchState->ChunkLocalData.NumOutputPerDataSet[DataSetIdx];
-					uint32 InstanceOffset     = BatchState->ChunkLocalData.StartingOutputIdxPerDataSet[DataSetIdx];
-					uint32 RegTypeOffset      = VVMState->DataSets[DataSetIdx].OutputRegisterTypeOffsets[RegType];
-					uint32 **OutputBuffers    = (uint32 **)VVMState->DataSets[DataSetIdx].OutputRegisters.GetData();
-					uint32 *DstReg            = OutputBuffers[RegTypeOffset + VecIndices[3]] + InstanceOffset;
-					uint32 *SrcReg            = (uint32 *)VecReg[2];
-					if (NumOutputInstances == NumInstancesThisChunk)
-					{
-						if (RegInc[2] == 0) //setting from a constant
-						{ 
-							VVMMemSet32(DstReg, *SrcReg, NumOutputInstances);
-						}
-						else //copying from internal buffers
-						{
-							VVMMemCpy(DstReg, SrcReg, sizeof(uint32) * NumOutputInstances);
-						}
-					}
-					else //if we are discarding at least one instance we can't just copy the memory
-					{
-						for (uint32 i = 0; i < NumOutputInstances; ++i)
-						{
-							DstReg[i] = SrcReg[DstIdxReg[i & RegInc[1]] & RegInc[2]];
-						}
-					}
-					InsPtr += 8;
-				} break;
-				case EVectorVMOp::outputdata_half: {
-					uint8 RegType             = 2;
-					uint16 DataSetIdx         = VecIndices[0];
-					int *DstIdxReg            = (int *)VecReg[1];
-					uint32 NumOutputInstances = BatchState->ChunkLocalData.NumOutputPerDataSet[DataSetIdx];
-					uint32 InstanceOffset     = BatchState->ChunkLocalData.StartingOutputIdxPerDataSet[DataSetIdx];
-					uint32 RegTypeOffset      = VVMState->DataSets[DataSetIdx].OutputRegisterTypeOffsets[RegType];
-					uint16 **OutputBuffers    = (uint16 **)VVMState->DataSets[DataSetIdx].OutputRegisters.GetData();
-					uint16 *DstReg            = OutputBuffers[RegTypeOffset + VecIndices[3]] + InstanceOffset;
-					uint32 *SrcReg            = (uint32 *)VecReg[2];
-
-					for (uint32 i = 0; i < NumOutputInstances; ++i) //scalar, not ideal
-					{
-						FPlatformMath::StoreHalf(DstReg + i, ((float *)SrcReg)[DstIdxReg[i & RegInc[1]] & RegInc[2]]);
-					}
-					InsPtr += 8;
-				} break;
-				case EVectorVMOp::acquireindex:
-				{
-					uint32 NumOutputInstances = 0;
-					uint16 DataSetIdx         = VecIndices[0];
-					uint32 *Input             = (uint32 *)VecReg[1];
-					uint32 *Output            = (uint32 *)VecReg[2];
-					uint16 IncMask            = RegInc[1];
-					for (uint16 i = 0; i < NumInstancesThisChunk; ++i)
-					{
-						//since input and output can alias now we need to save the IncAmt to a temp value
-						uint32 IncAmt = Input[i & IncMask] >> 31; //-1 is keep, so we only need to check for the high bit
-						Output[NumOutputInstances] = i;
-						NumOutputInstances += IncAmt;
-					}
-					//the new VM's indicies are generated to support brachless write-gather for the output instructions (instead of an in-signal flag as the original bytecode intended)
-					//the above loop will write an invalid value into the last slot if we discard one or more instances.  This is normally okay, however if an update_id instruction is issued later, 
-					//we will write incorrect values into the free id table there.  To avoid this (and potentially other problems that may come up if the bytecode is expanded, we correct the final slot here.
-					if ((int)NumOutputInstances < NumInstancesThisChunk)
-					{
-						Output[NumOutputInstances] = NumInstancesThisChunk;
-					}
-					BatchState->ChunkLocalData.StartingOutputIdxPerDataSet[DataSetIdx] = VVMState->DataSets[DataSetIdx].InstanceOffset + FPlatformAtomics::InterlockedAdd(VVMState->NumOutputPerDataSet + DataSetIdx, NumOutputInstances);
-					BatchState->ChunkLocalData.NumOutputPerDataSet[DataSetIdx] += NumOutputInstances;
-					VVMSer_regUsed(VecIndices[2], VVMRegFlag_Int | VVMRegFlag_Index);
-					InsPtr += 8;
-				}
-				break;
-				case EVectorVMOp::external_func_call:
+			//regular inputs: float, int and half
+			for (int j = 0; j < 3; ++j)
+			{
+				int NumInputsThisType = ExecCtx->VVMState->InputDataSetOffsets[(i << 3) + j + 1] - ExecCtx->VVMState->InputDataSetOffsets[(i << 3) + j];
+				int TypeOffset = ExecCtx->DataSets[i].InputRegisterTypeOffsets[j];
+				for (int k = 0; k < NumInputsThisType; ++k)
 				{
-					FVectorVMExtFunctionData *ExtFnData = VVMState->ExtFunctionTable + *InsPtr;
-#if					defined(VVM_INCLUDE_SERIALIZATION) && !defined(VVM_SERIALIZE_NO_WRITE)
-					if (SerializeState && (SerializeState->Flags & VVMSer_SyncExtFns) && CmpSerializeState && SerializeState->NumInstances == CmpSerializeState->NumInstances && (CmpSerializeState->NumInstructions > SerializeState->NumInstructions || VVMSerGlobalChunkIdx != 0))
-					{
-						//If we hit this branch we are using the output from the comparision state instead of running the external function itself.
-						//AFAIK the VM is not speced to have the inputs and outputs in a particular order, and even if it is we shouldn't rely on
-						//3rd party external function writers to follow the spec.  Therefore we don't just sync what we think is output, we sync
-						//all temp registers that are used in the function.
-						int NumRegisters = ExtFnData->NumInputs + ExtFnData->NumOutputs;
-						FVectorVMSerializeInstruction *CmpIns = nullptr; //this instruction
-						FVectorVMOptimizeInstruction *OptIns = nullptr;
-						if (SerializeState->OptimizeCtx)
-						{
-							//instructions have been re-ordered, can't binary search, must linear search
-							for (uint32 i = 0 ; i < SerializeState->OptimizeCtx->Intermediate.NumInstructions; ++i)
-							{
-								if (SerializeState->OptimizeCtx->Intermediate.Instructions[i].PtrOffsetInOptimizedBytecode == (int)(VVMSerStartOpPtr - VVMState->Bytecode))
-								{
-									OptIns = SerializeState->OptimizeCtx->Intermediate.Instructions + i;
-									break;
-								}
-							}
-							if (OptIns)
-							{
-								for (uint32 i = 0; i < CmpSerializeState->NumInstructions; ++i)
-								{
-									if (OptIns->PtrOffsetInOrigBytecode == CmpSerializeState->Instructions[i].OpStart)
-									{
-										CmpIns = CmpSerializeState->Instructions + i;
-										break;
-									}
-								}
-							}
-						}
-						else
-						{
-							CmpIns = CmpSerializeState->Instructions + VVMSerNumInstructionsThisChunk;
-						}
-						if (OptIns && CmpIns)
-						{
-							for (int i = 0; i < NumRegisters; ++i)
-							{
-								uint16 DstRegIdx = ((uint16 *)(InsPtr + 2))[i];
-								if (!(DstRegIdx & 0x8000)) //high bit signifies constant, skip it
-								{
-									uint32 DstOffset = VVMState->NumConstBuffers + DstRegIdx * NumLoops;
-									uint32 SrcRegIdx = DstRegIdx;
-									if (OptIns && (SerializeState->Flags & VVMSer_OptimizedBytecode) && !(CmpSerializeState->Flags & VVMSer_OptimizedBytecode))
-									{
-										SrcRegIdx = ((uint16 *)(CmpSerializeState->Bytecode + OptIns->PtrOffsetInOrigBytecode + 2))[i] & 0x7FFF; //high bit is register in original bytecode
-									}
-									if (SrcRegIdx != 0x7FFF) //invalid register, skipped by external function in the execution
-									{ 
-										uint8 *Src = (uint8 *)(CmpIns->TempRegisters + SrcRegIdx * CmpSerializeState->NumInstances + StartInstanceThisChunk);
-										uint8 *Dst = (uint8 *)(BatchState->RegisterData + DstOffset);
-										VVMMemCpy(Dst, Src, sizeof(uint32) * NumInstancesThisChunk);
-									}
-								}
-							}
-						}
-						for (int i = 0; i < ExtFnData->NumOutputs; ++i)
-						{
-							uint16 RegIdx = ((uint16 *)(InsPtr + 2))[ExtFnData->NumInputs + i];
-							VVMSer_regUsed(RegIdx, 0); //assume float
-						}
-					} else {
-#					else //VVM_INCLUDE_SERIALIZATION
-					{
-#					endif //VVM_INCLUDE_SERIALIZATION
-						check(*InsPtr < VVMState->NumExtFunctions);
-						check((uint32)(ExtFnData->NumInputs + ExtFnData->NumOutputs) <= VVMState->MaxExtFnRegisters);
-
-						//first decode all of the registers this external function needs into the batch's Chunk Local Data
-						//skip the first index because it's the ExtFnIdx
-						BatchState->ChunkLocalData.ExtFnDecodedReg.RegData[0]    = VecReg[1];
-						BatchState->ChunkLocalData.ExtFnDecodedReg.RegData[1]    = VecReg[2];
-						BatchState->ChunkLocalData.ExtFnDecodedReg.RegData[2]    = VecReg[3];
-
-						BatchState->ChunkLocalData.ExtFnDecodedReg.RegInc[0]     = RegInc[1];
-						BatchState->ChunkLocalData.ExtFnDecodedReg.RegInc[1]     = RegInc[2];
-						BatchState->ChunkLocalData.ExtFnDecodedReg.RegInc[2]     = RegInc[3];
-						
-						for (int i = 3; i < ExtFnData->NumInputs + ExtFnData->NumOutputs; i += 4)
-						{
-							VVMDecodeInstructionRegisters(InsPtr + 2 + i * 2, BatchState->ChunkLocalData.ExtFnDecodedReg.RegData + i, BatchState->ChunkLocalData.ExtFnDecodedReg.RegInc + i);
-						}
-						FVectorVMExternalFunctionContextExperimental ExtFnCtx;
-
-						ExtFnCtx.RegisterData             = (uint32 **)BatchState->ChunkLocalData.ExtFnDecodedReg.RegData;
-						ExtFnCtx.RegInc                   = BatchState->ChunkLocalData.ExtFnDecodedReg.RegInc;
-						ExtFnCtx.RawVecIndices            = VecIndices + 1; //skip index 0, that's the function index
-
-						ExtFnCtx.RegReadCount             = 0;
-						ExtFnCtx.NumRegisters             = ExtFnData->NumInputs + ExtFnData->NumOutputs;
-
-						ExtFnCtx.StartInstance            = StartInstanceThisChunk;
-						ExtFnCtx.NumInstances             = NumInstancesThisChunk;
-						ExtFnCtx.NumLoops                 = NumLoops;
-						ExtFnCtx.PerInstanceFnInstanceIdx = 0;
-
-						ExtFnCtx.UserPtrTable             = VVMState->UserPtrTable;
-						ExtFnCtx.NumUserPtrs              = VVMState->NumUserPtrTable;
-						ExtFnCtx.RandStream               = &BatchState->RandStream;
-
-						ExtFnCtx.RandCounters             = BatchState->ChunkLocalData.RandCounters;
-						ExtFnCtx.DataSets                 = VVMState->DataSets;
-
-						//correct invalid registers
-						uint32 DummyRegCount = 0;
-						for (int i = 0; i < ExtFnData->NumInputs + ExtFnData->NumOutputs; ++i) {
-							if (ExtFnCtx.RawVecIndices[i] == 0xFFFF) {
-								ExtFnCtx.RegInc[i] = 0;
-								ExtFnCtx.RegisterData[i] = (uint32 *)(BatchState->ChunkLocalData.ExtFnDecodedReg.DummyRegs + DummyRegCount++);
-							}
-						}
-						check(DummyRegCount <= VVMState->NumDummyRegsReq);
-#if VECTORVM_SUPPORTS_LEGACY
-						FVectorVMExternalFunctionContext ProxyContext(ExtFnCtx);
-						ExtFnData->Function->Execute(ProxyContext);
-#else
-						ExtFnData->Function->Execute(ExtFnCtx);
-#endif
-					}
-					InsPtr += 2 + ((ExtFnData->NumInputs + ExtFnData->NumOutputs) << 1);
+					int RemapIdx = ExecCtx->VVMState->InputDataSetOffsets[(i << 3) + j] + k;
+					ExecCtx->VVMState->InputMapCacheIdx[InputCounter] = i;
+					ExecCtx->VVMState->InputMapCacheSrc[InputCounter] = TypeOffset + ExecCtx->VVMState->InputRemapTable[RemapIdx];
+					++InputCounter;
 				}
-				break;
-				case EVectorVMOp::exec_index:
+			}
+
+			if (ExecCtx->VVMState->InputDataSetOffsets[(i << 3) + 7] > 0)
+			{
+				//no advance inputs: float, int and half
+				//no advance inputs point directly after the constant buffers
+				for (int j = 0; j < 3; ++j)
 				{
-					int RegIdx = NumLoops * *((uint16 *)InsPtr);
-					VVMSer_instruction(1, 0);
-					VectorRegister4i StartInstance4 = VectorIntAdd(VectorIntSet1(StartInstanceThisChunk), VVM_m128iConst(ZeroOneTwoThree));
-					for (int i = 0; i < NumLoops; ++i)
+					int NumInputsThisType = ExecCtx->VVMState->InputDataSetOffsets[(i << 3) + j + 4] - ExecCtx->VVMState->InputDataSetOffsets[(i << 3) + j + 3];
+					int TypeOffset = ExecCtx->DataSets[i].InputRegisterTypeOffsets[j];
+					for (int k = 0; k < NumInputsThisType; ++k)
 					{
-						VectorRegister4i i4 = VectorIntSet1(i);
-						VectorRegister4i v4 = VectorIntAdd(StartInstance4, VectorShiftLeftImm(i4, 2));
-						VecReg[0][i].i = v4;
+						int RemapIdx     = ExecCtx->VVMState->InputDataSetOffsets[(i << 3) + j + 3] + k;
+						ExecCtx->VVMState->InputMapCacheIdx[InputCounter] = i;
+						ExecCtx->VVMState->InputMapCacheSrc[InputCounter] = (TypeOffset + ExecCtx->VVMState->InputRemapTable[RemapIdx]) | 0x8000 | (((j & 2) << 13)); //high bit is no advance input, 2nd high bit is whether it's half
+						++InputCounter;
 					}
-					InsPtr += 2;
 				}
-				break;
-				case EVectorVMOp::noise2D:							check(false);                   break;
-				case EVectorVMOp::noise3D:							check(false);                   break;
-				case EVectorVMOp::enter_stat_scope:					InsPtr += 2;                    break;
-				case EVectorVMOp::exit_stat_scope:                                                  break;
-				case EVectorVMOp::update_id:
-				{
-					VVMSer_instruction(1, 2);
-					uint32 DataSetIdx     = VecIndices[0];
-					check(DataSetIdx < (uint32)VVMState->DataSets.Num());
-					FDataSetMeta *DataSet = &VVMState->DataSets[DataSetIdx];
-					int32 *R1             = (int32 *)VecReg[1];
-					int32 *R2             = (int32 *)VecReg[2];
-
-					check(DataSet->IDTable);
-					check(DataSet->IDTable->Num() >= DataSet->InstanceOffset + StartInstanceThisChunk + NumInstancesThisChunk);
-
-					int NumOutputInstances = BatchState->ChunkLocalData.NumOutputPerDataSet[DataSetIdx];
-					int NumFreed           = NumInstancesThisChunk - BatchState->ChunkLocalData.NumOutputPerDataSet[DataSetIdx];
-
-					//compute this chunk's MaxID
-					int MaxID = -1;
-					if (NumOutputInstances > 4)
-					{
-						int NumOutput4 = (int)(((((uint32)NumOutputInstances + 3U) & ~3U) - 1) >> 2);
-						VectorRegister4i Max4 = VectorIntSet1(-1);
-						check(((size_t)VecReg[1] & 0xF) == 0); //this must come from a register, we don't fuse input on update_id so we know it's aligned.  we can use pcmpgtd xmm, [ptr]
-						for (int i = 0; i < NumOutput4; ++i)
-						{
-							Max4 = VectorIntXor(VecReg[1][i].i, VectorIntAnd(VectorIntCompareGT(Max4, VecReg[1][i].i), VectorIntXor(Max4, VecReg[1][i].i)));
-						}
-						VectorRegister4i Last4 = VectorIntLoad(R1 + NumOutputInstances - 4);
-						Max4 = VectorIntXor(Last4, VectorIntAnd(VectorIntCompareGT(Max4, Last4), VectorIntXor(Max4, Last4)));
-						int M4[4];
-						VectorIntStore(Max4, M4);
-						int m0 = M4[0] > M4[1] ? M4[0] : M4[1];
-						int m1 = M4[2] > M4[3] ? M4[2] : M4[3];
-						int m = m0 > m1 ? m0 : m1;
-						if (m > MaxID)
-						{
-							MaxID = m;
-						}
-					}
-					else
-					{
-						for (int i = 0; i < NumOutputInstances; ++i)
-						{
-							if (R1[i] > MaxID)
-							{
-								MaxID = R1[i];
-							}
-						}
-					}
-					
-					// Update the actual index for this ID.  No thread safety is required as this ID slot can only ever be written by this instance
-					// The index passed into this function is the same as that given to the output* instructions
-					for (int i = 0; i < NumOutputInstances; ++i)
-					{
-						(*DataSet->IDTable)[R1[R2[i]]] = BatchState->ChunkLocalData.StartingOutputIdxPerDataSet[DataSetIdx] + i; //BatchState->ChunkLocalData.StartingOutputIdxPerDataSet[DataSetIdx] already has DataSet->InstanceOffset added to it
-					}
+			}
+		}
+		check(InputCounter == ExecCtx->VVMState->NumInputBuffers);
+	}
+}
 
-					//Write the freed indices to the free table.
-					if (NumFreed > 0)
-					{
-						int StartNumFreed = FPlatformAtomics::InterlockedAdd((volatile int32 *)DataSet->NumFreeIDs, NumFreed);
-						int32 *FreeTableStart = DataSet->FreeIDTable->GetData() + StartNumFreed;
-						int c = 0;
-						int FreeCount = 0;
-						while (FreeCount < NumFreed)
-						{
-							check(c < NumInstancesThisChunk);
-							int d = R2[c] - c - FreeCount; //check for a gap in the write index and the counter... if nothing is freed then the write index matches the counter
-							if (d > 0)
-							{
-								VVMMemCpy(FreeTableStart + FreeCount, R1 + FreeCount + c, sizeof(int32) * d);
-								FreeCount += d;
-							}
-							++c;
-						}
-						check(FreeCount == NumFreed);
-					}
+VECTORVM_API void ExecVectorVMState(FVectorVMExecContext *ExecCtx, FVectorVMSerializeState *SerializeState, FVectorVMSerializeState *CmpSerializeState)
+{
+	if (ExecCtx->ExtFunctionTable.Num() != ExecCtx->VVMState->NumExtFunctions) {
+		return;
+	}
+	for (uint32 i = 0; i < ExecCtx->VVMState->NumExtFunctions; ++i) {
+		ExecCtx->VVMState->ExtFunctionTable[i].Function = ExecCtx->ExtFunctionTable[i];
+	}
+	for (uint32 i = 0; i < ExecCtx->VVMState->NumOutputDataSets; ++i) {
+		ExecCtx->VVMState->NumOutputPerDataSet[i] = 0;
+	}
 
-					//Set the DataSet's MaxID if this chunk's MaxID is bigger
-					if (MaxID != -1)
-					{
-						int SanityCount = 0;
-						do {
-							int OldMaxID = FPlatformAtomics::AtomicRead(DataSet->MaxUsedID);
-							if (MaxID <= OldMaxID)
-							{
-								break;
-							}
-							int NewMaxID = FPlatformAtomics::InterlockedCompareExchange((volatile int32 *)DataSet->MaxUsedID, MaxID, OldMaxID);
-							if (NewMaxID == OldMaxID)
-							{
-								break;
-							}
-						} while (SanityCount++ < (1 << 30));
-						VVMDebugBreakIf(SanityCount > (1 << 30) - 1);
-					}
-					InsPtr += 6;
-				}
-				break;
-				case EVectorVMOp::acquire_id:
-				{
-					VVMSer_instruction(1, 2);
-					uint32 DataSetIdx = VecIndices[0];
-					check(DataSetIdx < (uint32)VVMState->DataSets.Num());
-					FDataSetMeta *DataSet = &VVMState->DataSets[DataSetIdx];
-					
-					{ //1. Get the free IDs into the temp register
-						int SanityCount = 0;
-						do
-						{
-							int OldNumFreeIDs = FPlatformAtomics::AtomicRead(DataSet->NumFreeIDs);
-							check(OldNumFreeIDs >= NumInstancesThisChunk);
-							//this is reverse-order from the original VM but it shouldn't matter since these are just re-used indices.
-							//VVMMemCpy(VecReg[1], DataSet->FreeIDTable->GetData() + OldNumFreeIDs - NumInstancesThisChunk, sizeof(int32) * NumInstancesThisChunk); //pull off the last added FreeIDs
-							int *OutPtr = (int *)VecReg[1];
-							int *InPtr  = DataSet->FreeIDTable->GetData() + OldNumFreeIDs - NumInstancesThisChunk;
-							for (int i = 0; i < NumInstancesThisChunk; ++i)
-							{
-								OutPtr[i] = InPtr[NumInstancesThisChunk - i - 1];
-							}
-							int NewNumFreeIDs = FPlatformAtomics::InterlockedCompareExchange((volatile int32 *)DataSet->NumFreeIDs, OldNumFreeIDs - NumInstancesThisChunk, OldNumFreeIDs);
-							if (NewNumFreeIDs == OldNumFreeIDs)
-							{
-								break;
-							}
-						} while (SanityCount++ < (1 << 30));
-						VVMDebugBreakIf(SanityCount > (1 << 30) - 1);
-					}
-					{ //2. append the IDs we acquired in step 1 to the end of the free table array, representing spawned IDs
-						//FreeID table is write-only as far as this invocation of the VM is concerned, so the interlocked add w/o filling
-						//in the data is fine
-						int StartNumSpawned = FPlatformAtomics::InterlockedAdd(DataSet->NumSpawnedIDs, NumInstancesThisChunk) + NumInstancesThisChunk;
-						check(StartNumSpawned <= DataSet->FreeIDTable->Max());
-						VVMMemCpy(DataSet->FreeIDTable->GetData() + DataSet->FreeIDTable->Max() - StartNumSpawned, VecReg[1], sizeof(int32) * NumInstancesThisChunk);
-					}
-					//3. set the tag
-					VVMMemSet32(VecReg[2], DataSet->IDAcquireTag, NumInstancesThisChunk);
-					InsPtr += 6;
-				}
-				break;
-				case EVectorVMOp::fused_input1_1: //op has 1 input operand and it's being overwritten to an input
-				{
-					OpCode                    = (EVectorVMOp)InsPtr[4];
-					uint8 RegType             = InsPtr[5];
-					uint16 DataSetIdx         = *(uint16 *)(InsPtr + 6);
-					uint16 InputRegIdx        = VecIndices[0];
-					uint32 InstanceOffset     = VVMState->DataSets[DataSetIdx].InstanceOffset;
-					uint32 InputRegTypeOffset = VVMState->DataSets[DataSetIdx].InputRegisterTypeOffsets[RegType];
-					uint32 **InputBuffers     = (uint32 **)VVMState->DataSets[DataSetIdx].InputRegisters.GetData();
-					uint32 *InputPtr          = InputBuffers[InputRegIdx + InputRegTypeOffset] + StartInstanceThisChunk + InstanceOffset;
-					VecReg[0]                 = (FVecReg *)InputPtr;
-					InsPtr += 4;
-					goto OpCodeSwitch;
-				}
-				break;
-				case EVectorVMOp::fused_input2_1: //op has 2 input operands, register 0 is being overwritten to an input (intentional fallthrough)
-				case EVectorVMOp::fused_input2_2: //op has 2 input operands, register 1 is being overwritten to an input
-				{
-					int RegToSwitchToInput     = (int)OpCode - (int)EVectorVMOp::fused_input2_1;
-					OpCode                     = (EVectorVMOp)InsPtr[6];
-					uint8 RegType              = InsPtr[7];
-					uint16 DataSetIdx          = *(uint16 *)(InsPtr + 8);
-					uint16 InputRegIdx         = VecIndices[RegToSwitchToInput];
-					uint32 InstanceOffset      = VVMState->DataSets[DataSetIdx].InstanceOffset;
-					uint32 InputRegTypeOffset  = VVMState->DataSets[DataSetIdx].InputRegisterTypeOffsets[RegType];
-					uint32 **InputBuffers      = (uint32 **)VVMState->DataSets[DataSetIdx].InputRegisters.GetData();
-					uint32 *InputPtr           = InputBuffers[InputRegIdx + InputRegTypeOffset] + StartInstanceThisChunk + InstanceOffset;
-					VecReg[RegToSwitchToInput] = (FVecReg *)InputPtr;
-					InsPtr += 4;
-					goto OpCodeSwitch;
-				}
-				break;
-				case EVectorVMOp::fused_input2_3: //op has 2 input operands, register 0 and 1 are being overwritten to an input
-				{
-					OpCode                       = (EVectorVMOp)InsPtr[6];
-					uint8 RegType[2]             = { InsPtr[7], InsPtr[10] };
-					uint16 DataSetIdx[2]         = { *(uint16 *)(InsPtr + 8), *(uint16 *)(InsPtr + 11) };
-					int32 InstanceOffset[2]      = { VVMState->DataSets[DataSetIdx[0]].InstanceOffset, VVMState->DataSets[DataSetIdx[1]].InstanceOffset };
-					uint32 InputRegTypeOffset[2] = { VVMState->DataSets[DataSetIdx[0]].InputRegisterTypeOffsets[RegType[0]], VVMState->DataSets[DataSetIdx[1]].InputRegisterTypeOffsets[RegType[1]] };
-					uint32 **InputBuffers[2]     = { (uint32 **)VVMState->DataSets[DataSetIdx[0]].InputRegisters.GetData(), (uint32 **)VVMState->DataSets[DataSetIdx[1]].InputRegisters.GetData() };
-					uint32 *InputPtr[2]          = { InputBuffers[0][VecIndices[0] + InputRegTypeOffset[0]] + StartInstanceThisChunk + InstanceOffset[0], InputBuffers[1][VecIndices[1] + InputRegTypeOffset[1]] + StartInstanceThisChunk + InstanceOffset[1] };
-					VecReg[0]                    = (FVecReg *)InputPtr[0];
-					VecReg[1]                    = (FVecReg *)InputPtr[1];
-					InsPtr += 7;
-					goto OpCodeSwitch;
-				}
-				break;
-				case EVectorVMOp::fused_input3_1: //op has 3 input operands, register 0 is being overwritten to an input (intentional fallthrough)
-				case EVectorVMOp::fused_input3_2: //op has 3 input operands, register 1 is being overwritten to an input (intentional fallthrough)
-				case EVectorVMOp::fused_input3_4: //op has 3 input operands, register 2 is being overwritten to an input
-				{
-					int RegToSwitchToInput     = (int)OpCode - (int)EVectorVMOp::fused_input3_1;
-					check(RegToSwitchToInput >= 0 && RegToSwitchToInput <= 2);
-					OpCode                     = (EVectorVMOp)InsPtr[8];
-					uint8 RegType              = InsPtr[9];
-					uint16 DataSetIdx          = *(uint16 *)(InsPtr + 10);
-					uint16 InputRegIdx         = VecIndices[RegToSwitchToInput];
-					uint32 InstanceOffset      = VVMState->DataSets[DataSetIdx].InstanceOffset;
-					uint32 InputRegTypeOffset  = VVMState->DataSets[DataSetIdx].InputRegisterTypeOffsets[RegType];
-					uint32 **InputBuffers      = (uint32 **)VVMState->DataSets[DataSetIdx].InputRegisters.GetData();
-					uint32 *InputPtr           = InputBuffers[InputRegIdx + InputRegTypeOffset] + StartInstanceThisChunk + InstanceOffset;
-					VecReg[RegToSwitchToInput] = (FVecReg *)InputPtr;
-					InsPtr += 4;
-					goto OpCodeSwitch;
-				}
-				break;
-				case EVectorVMOp::fused_input3_3: //op has 3 input operands, register 0 and 1 are being overwritten to inputs (intentional fallthrough)
-				case EVectorVMOp::fused_input3_5: //op has 3 input operands, register 0 and 2 are being overwritten to inputs
-				{
-					int RegIdx2                  = 1 + (int)OpCode - (int)EVectorVMOp::fused_input3_3;
-					OpCode                       = (EVectorVMOp)InsPtr[8];
-					uint8 RegType[2]             = { InsPtr[9], InsPtr[12] };
-					uint16 DataSetIdx[2]         = { *(uint16 *)(InsPtr + 10), *(uint16 *)(InsPtr + 13) };
-					int32 InstanceOffset[2]      = { VVMState->DataSets[DataSetIdx[0]].InstanceOffset, VVMState->DataSets[DataSetIdx[1]].InstanceOffset };
-					uint32 InputRegTypeOffset[2] = { VVMState->DataSets[DataSetIdx[0]].InputRegisterTypeOffsets[RegType[0]], VVMState->DataSets[DataSetIdx[1]].InputRegisterTypeOffsets[RegType[1]] };
-					uint32 **InputBuffers[2]     = { (uint32 **)VVMState->DataSets[DataSetIdx[0]].InputRegisters.GetData(), (uint32 **)VVMState->DataSets[DataSetIdx[1]].InputRegisters.GetData() };
-					uint32 *InputPtr[2]          = { InputBuffers[0][VecIndices[0] + InputRegTypeOffset[0]] + StartInstanceThisChunk + InstanceOffset[0], InputBuffers[1][VecIndices[RegIdx2] + InputRegTypeOffset[1]] + StartInstanceThisChunk + InstanceOffset[1] };
-					VecReg[0]       = (FVecReg *)InputPtr[0];
-					VecReg[RegIdx2] = (FVecReg *)InputPtr[1];
-					InsPtr += 7;
-					goto OpCodeSwitch;
-				}
-				break;
-				case EVectorVMOp::fused_input3_6: //op has 3 input operands, register 1 and 2 are being overwritten to inputs
-				{
-					OpCode                       = (EVectorVMOp)InsPtr[8];
-					uint8 RegType[2]             = { InsPtr[9], InsPtr[12] };
-					uint16 DataSetIdx[2]         = { *(uint16 *)(InsPtr + 10), *(uint16 *)(InsPtr + 13) };
-					int32 InstanceOffset[2]      = { VVMState->DataSets[DataSetIdx[0]].InstanceOffset, VVMState->DataSets[DataSetIdx[1]].InstanceOffset };
-					uint32 InputRegTypeOffset[2] = { VVMState->DataSets[DataSetIdx[0]].InputRegisterTypeOffsets[RegType[0]], VVMState->DataSets[DataSetIdx[1]].InputRegisterTypeOffsets[RegType[1]] };
-					uint32 **InputBuffers[2]     = { (uint32 **)VVMState->DataSets[DataSetIdx[0]].InputRegisters.GetData(), (uint32 **)VVMState->DataSets[DataSetIdx[1]].InputRegisters.GetData() };
-					uint32 *InputPtr[2]          = { InputBuffers[0][VecIndices[1] + InputRegTypeOffset[0]] + StartInstanceThisChunk + InstanceOffset[0], InputBuffers[1][VecIndices[2] + InputRegTypeOffset[1]] + StartInstanceThisChunk + InstanceOffset[1] };
-					VecReg[1] = (FVecReg *)InputPtr[0];
-					VecReg[2] = (FVecReg *)InputPtr[1];
-					InsPtr += 7;
-					goto OpCodeSwitch;
-				}
-				break;
-				case EVectorVMOp::fused_input3_7: //op has 3 input operands, register 1, 2, and 3 are all being overwritten to inputs
-				{
-					OpCode                       = (EVectorVMOp)InsPtr[8];
-					uint8 RegType[3]             = { InsPtr[9], InsPtr[12], InsPtr[15] };
-					uint16 DataSetIdx[3]         = { *(uint16 *)(InsPtr + 10), *(uint16 *)(InsPtr + 13), *(uint16 *)(InsPtr + 16) };
-					int32 InstanceOffset[3]      = { VVMState->DataSets[DataSetIdx[0]].InstanceOffset, 
-						                             VVMState->DataSets[DataSetIdx[1]].InstanceOffset, 
-						                             VVMState->DataSets[DataSetIdx[2]].InstanceOffset
-					                               };
-					uint32 InputRegTypeOffset[3] = { VVMState->DataSets[DataSetIdx[0]].InputRegisterTypeOffsets[RegType[0]], 
-						                             VVMState->DataSets[DataSetIdx[1]].InputRegisterTypeOffsets[RegType[1]], 
-						                             VVMState->DataSets[DataSetIdx[2]].InputRegisterTypeOffsets[RegType[2]] 
-					                               };
-					uint32 **InputBuffers[3]     = { (uint32 **)VVMState->DataSets[DataSetIdx[0]].InputRegisters.GetData(), 
-						                             (uint32 **)VVMState->DataSets[DataSetIdx[1]].InputRegisters.GetData(), 
-						                             (uint32 **)VVMState->DataSets[DataSetIdx[2]].InputRegisters.GetData() 
-					                               };
-					uint32 *InputPtr[3]          = { InputBuffers[0][VecIndices[0] + InputRegTypeOffset[0]] + StartInstanceThisChunk + InstanceOffset[0], 
-						                             InputBuffers[1][VecIndices[1] + InputRegTypeOffset[1]] + StartInstanceThisChunk + InstanceOffset[1], 
-						                             InputBuffers[2][VecIndices[2] + InputRegTypeOffset[2]] + StartInstanceThisChunk + InstanceOffset[2]
-					                               };
-					VecReg[0] = (FVecReg *)InputPtr[0];
-					VecReg[1] = (FVecReg *)InputPtr[1];
-					VecReg[2] = (FVecReg *)InputPtr[2];
-					InsPtr += 10;
-					goto OpCodeSwitch;
-				}
-				break;
-				case EVectorVMOp::copy_to_output:
-				{
-					uint16 OutputDataSetIdx   = VecIndices[0];
-					uint16 InputDataSetIdx    = VecIndices[1];
-					uint16 OutputDstIdxRegIdx = VecIndices[2];
-					uint8 RegType             = InsPtr[6];
-					uint8 Count               = InsPtr[7];
-					InsPtr += 8;
-
-					//Output
-					uint32 NumOutputInstances   = BatchState->ChunkLocalData.NumOutputPerDataSet[OutputDataSetIdx];
-					uint32 OutputInstanceOffset = BatchState->ChunkLocalData.StartingOutputIdxPerDataSet[OutputDataSetIdx];
-					uint32 OutputRegTypeOffset  = VVMState->DataSets[OutputDataSetIdx].OutputRegisterTypeOffsets[RegType];
-					uint32 **OutputBuffers      = (uint32 **)VVMState->DataSets[OutputDataSetIdx].OutputRegisters.GetData();
-
-					//Input
-					uint32 InputInstanceOffset = VVMState->DataSets[InputDataSetIdx].InstanceOffset;
-					uint32 InputRegTypeOffset  = VVMState->DataSets[InputDataSetIdx].InputRegisterTypeOffsets[RegType];
-					uint32 **InputBuffers      = (uint32 **)VVMState->DataSets[InputDataSetIdx].InputRegisters.GetData();
-
-					int *DstIdxReg             = (int *)VecReg[2];
-					int StartSrcIndex          = StartInstanceThisChunk + InputInstanceOffset;
-					uint16 *IdxInsPtr          = (uint16 *)InsPtr;
-					if (NumOutputInstances == BatchState->NumInstances)
-					{
-						//if we're writing the same number of inputs as outputs then we can just memcpy
-						for (int i = 0; i < (int)Count; ++i)
-						{
-							uint16 OutputDstIdx = IdxInsPtr[(i << 1) + 0];
-							uint16 InputSrcIdx  = IdxInsPtr[(i << 1) + 1];
-							uint32 *SrcBuffer   = InputBuffers [InputSrcIdx  + InputRegTypeOffset ] + StartSrcIndex;
-							uint32 *DstBuffer   = OutputBuffers[OutputDstIdx + OutputRegTypeOffset] + OutputInstanceOffset;
-							VVMMemCpy(DstBuffer, SrcBuffer, sizeof(uint32) * NumOutputInstances);
-						}
-					}
-					else
-					{
-						//if we are discarding at least one instance we can't just copy the memory and need to use the index generated in acquire_index
-						for (int i = 0; i < (int)Count; ++i)
-						{
-							uint16 OutputDstIdx = IdxInsPtr[(i << 1) + 0];
-							uint16 InputSrcIdx  = IdxInsPtr[(i << 1) + 1];
-							uint32 *SrcBuffer   = InputBuffers [InputSrcIdx  + InputRegTypeOffset ] + StartSrcIndex;
-							uint32 *DstBuffer   = OutputBuffers[OutputDstIdx + OutputRegTypeOffset] + OutputInstanceOffset;
-							for (uint32 j = 0; j < NumOutputInstances; ++j)
-							{
-								DstBuffer[j] = SrcBuffer[DstIdxReg[j]];
-							}
-						}
-					}
-					InsPtr += (int)Count * 4;
-				}
-				break;
-				case EVectorVMOp::output_batch2:
-				{
-					int *DstIdxReg            = (int *)VecReg[0];
-					uint16 DataSetIdx         = VecIndices[3];
-					uint32 RegTypeOffset      = VVMState->DataSets[DataSetIdx].OutputRegisterTypeOffsets[InsPtr[12]];
-					uint32 **OutputBuffers    = (uint32 **)VVMState->DataSets[DataSetIdx].OutputRegisters.GetData();
-					uint32 InstanceOffset     = BatchState->ChunkLocalData.StartingOutputIdxPerDataSet[DataSetIdx];
-					uint32 NumOutputInstances = BatchState->ChunkLocalData.NumOutputPerDataSet[DataSetIdx];
-
-					uint32 *DstReg[2]         = { OutputBuffers[RegTypeOffset + VecIndices[4]] + InstanceOffset, 
-						                          OutputBuffers[RegTypeOffset + VecIndices[5]] + InstanceOffset
-												};
-					uint32 *SrcReg[2]         = { (uint32 *)VecReg[1], 
-					                              (uint32 *)VecReg[2] };
-					
-					if (NumOutputInstances == BatchState->NumInstances)
-					{
-						if (RegInc[1] == 0)
-						{
-							VVMMemSet32(DstReg[0], *SrcReg[0], NumOutputInstances);
-						}
-						else
-						{
-							VVMMemCpy(DstReg[0], SrcReg[0], sizeof(uint32) * NumOutputInstances);
-						}
-						if (RegInc[2] == 0)
-						{
-							VVMMemSet32(DstReg[1], *SrcReg[1], NumOutputInstances);
-						}
-						else
-						{
-							VVMMemCpy(DstReg[1], SrcReg[1], sizeof(uint32) * NumOutputInstances);
-						}
-					}
-					else
-					{
-						for (uint32 i = 0; i < NumOutputInstances; ++i)
-						{
-							DstReg[0][i] = SrcReg[0][DstIdxReg[i & RegInc[0]] & RegInc[1]];
-							DstReg[1][i] = SrcReg[1][DstIdxReg[i & RegInc[0]] & RegInc[2]];
-						}
-					}
-					InsPtr += 13;
-				}
-				break;
-				case EVectorVMOp::output_batch3:
-				{
-					int *DstIdxReg            = (int *)VecReg[0];
-					uint16 DataSetIdx         = VecIndices[4];
-					uint32 RegTypeOffset      = VVMState->DataSets[DataSetIdx].OutputRegisterTypeOffsets[InsPtr[16]];
-					uint32 **OutputBuffers    = (uint32 **)VVMState->DataSets[DataSetIdx].OutputRegisters.GetData();
-					uint32 InstanceOffset     = BatchState->ChunkLocalData.StartingOutputIdxPerDataSet[DataSetIdx];
-					uint32 NumOutputInstances = BatchState->ChunkLocalData.NumOutputPerDataSet[DataSetIdx];
-
-					uint32 *DstReg[3]         = { OutputBuffers[RegTypeOffset + VecIndices[5]] + InstanceOffset, 
-						                          OutputBuffers[RegTypeOffset + VecIndices[6]] + InstanceOffset,
-					                              OutputBuffers[RegTypeOffset + VecIndices[7]] + InstanceOffset };
-					uint32 *SrcReg[3]         = { (uint32 *)VecReg[1], 
-					                              (uint32 *)VecReg[2], 
-					                              (uint32 *)VecReg[3] };
-
-					if (NumOutputInstances == BatchState->NumInstances)
-					{
-						if (RegInc[1] == 0)
-						{
-							VVMMemSet32(DstReg[0], *SrcReg[0], NumOutputInstances);
-						}
-						else
-						{
-							VVMMemCpy(DstReg[0], SrcReg[0], sizeof(uint32) * NumOutputInstances);
-						}
-						if (RegInc[2] == 0)
-						{
-							VVMMemSet32(DstReg[1], *SrcReg[1], NumOutputInstances);
-						}
-						else
-						{
-							VVMMemCpy(DstReg[1], SrcReg[1], sizeof(uint32) * NumOutputInstances);
-						}
-						if (RegInc[3] == 0)
-						{
-							VVMMemSet32(DstReg[2], *SrcReg[2], NumOutputInstances);
-						}
-						else
-						{
-							VVMMemCpy(DstReg[2], SrcReg[2], sizeof(uint32) * NumOutputInstances);
-						}
-					}
-					else
-					{
-						for (uint32 i = 0; i < NumOutputInstances; ++i)
-						{
-							DstReg[0][i] = SrcReg[0][DstIdxReg[i & RegInc[0]] & RegInc[1]];
-							DstReg[1][i] = SrcReg[1][DstIdxReg[i & RegInc[0]] & RegInc[2]];
-							DstReg[2][i] = SrcReg[2][DstIdxReg[i & RegInc[0]] & RegInc[3]];
-						}
-					}
-					InsPtr += 17;
-				}
-				break;
-				case EVectorVMOp::output_batch4:
+	//cache the mappings from niagara data buffers to internal minimized set
+	if (!(ExecCtx->VVMState->Flags & VVMFlag_DataMapCacheSetup)) {
+		VVMBuildMapTableCaches(ExecCtx);
+		ExecCtx->VVMState->Flags |= VVMFlag_DataMapCacheSetup;
+	}
+
+	for (uint32 i = 0; i < ExecCtx->VVMState->NumConstBuffers; ++i) {
+		ExecCtx->VVMState->ConstantBuffers[i].i = VectorIntSet1(((uint32 *)ExecCtx->ConstantTableData[ExecCtx->VVMState->ConstMapCacheIdx[i]])[ExecCtx->VVMState->ConstMapCacheSrc[i]]);
+	}
+
+	//if the number of instances hasn't changed since the last exec, we don't have to re-compute the internal state
+	if (ExecCtx->NumInstances == ExecCtx->VVMState->NumInstancesExecCached) 
+	{
+		ExecCtx->Internal.NumBytesRequiredPerBatch          = ExecCtx->VVMState->ExecCtxCache.NumBytesRequiredPerBatch;
+		ExecCtx->Internal.PerBatchRegisterDataBytesRequired = ExecCtx->VVMState->ExecCtxCache.PerBatchRegisterDataBytesRequired;
+		ExecCtx->Internal.NumBatches                        = ExecCtx->VVMState->ExecCtxCache.NumBatches;
+		ExecCtx->Internal.MaxChunksPerBatch                 = ExecCtx->VVMState->ExecCtxCache.MaxChunksPerBatch;
+		ExecCtx->Internal.MaxInstancesPerChunk              = ExecCtx->VVMState->ExecCtxCache.MaxInstancesPerChunk;
+	}
+	else
+	{ //calculate Batch & Chunk division and all internal execution state required before executing
+		static const uint32 MaxChunksPerBatch                              = 4; //*MUST BE POW 2* arbitrary 4 chunks per batch... this is harder to load balance because it depends on CPU cores available during execution
+		static_assert(MaxChunksPerBatch > 0 && (MaxChunksPerBatch & (MaxChunksPerBatch - 1)) == 0);
+
+		size_t PageSizeInBytes = (uint64_t)GVVMPageSizeInKB << 10;
+
+		size_t PerBatchRegisterDataBytesRequired                           = 0;
+		int NumBatches                                                     = 1;
+		int NumChunksPerBatch                                              = (int)MaxChunksPerBatch;
+		uint32 MaxLoopsPerChunk                                            = 0;
+		{ //compute the number of bytes required per batch
+			const uint32 TotalNumLoopsRequired              = VVM_MAX(((uint32)ExecCtx->NumInstances + 3) >> 2, 1);
+			const size_t NumBytesRequiredPerLoop            = VVM_REG_SIZE * ExecCtx->VVMState->NumTempRegisters;
+			if (ExecCtx->VVMState->BatchOverheadSize + 64 > (uint32)PageSizeInBytes)
+			{
+				//either the chunk size is way too small, or there's an insane number of consts or data sets required... we just revert to the previous VM's default
+				MaxLoopsPerChunk                            = 128 >> 2;
+				NumBatches                                  = (TotalNumLoopsRequired + MaxLoopsPerChunk - 1) / MaxLoopsPerChunk;
+			}
+			else
+			{
+				size_t NumBytesPerBatchAvailableForTempRegs = PageSizeInBytes - ExecCtx->VVMState->BatchOverheadSize;
+				size_t TotalNumLoopBytesRequired            = VVM_ALIGN(TotalNumLoopsRequired, MaxChunksPerBatch) * NumBytesRequiredPerLoop;
+				if (NumBytesPerBatchAvailableForTempRegs < TotalNumLoopBytesRequired)
 				{
-					uint16 *OutputIndices     = VecIndices + 4;
-					uint16 DataSetIdx         = VecIndices[8];
-					uint16 IdxRegIdx          = VecIndices[9];
-					int *DstIdxReg            = (int *)(BatchState->RegisterData + VVMState->NumConstBuffers + NumLoops * IdxRegIdx); //guaranteed by the optimizer to be a temp register and not a const
-					uint32 RegTypeOffset      = VVMState->DataSets[DataSetIdx].OutputRegisterTypeOffsets[InsPtr[20]];
-					uint32 **OutputBuffers    = (uint32 **)VVMState->DataSets[DataSetIdx].OutputRegisters.GetData();
-					uint32 InstanceOffset     = BatchState->ChunkLocalData.StartingOutputIdxPerDataSet[DataSetIdx];
-					uint32 NumOutputInstances = BatchState->ChunkLocalData.NumOutputPerDataSet[DataSetIdx];
-
-					uint32 *DstReg[4]         = { OutputBuffers[RegTypeOffset + OutputIndices[0]] + InstanceOffset, 
-						                          OutputBuffers[RegTypeOffset + OutputIndices[1]] + InstanceOffset, 
-						                          OutputBuffers[RegTypeOffset + OutputIndices[2]] + InstanceOffset, 
-						                          OutputBuffers[RegTypeOffset + OutputIndices[3]] + InstanceOffset };
-
-					uint32 *SrcReg[4]         = { (uint32 *)VecReg[0], 
-						                          (uint32 *)VecReg[1], 
-						                          (uint32 *)VecReg[2], 
-						                          (uint32 *)VecReg[3] };
-					
-					if (NumOutputInstances == BatchState->NumInstances)
+					//Not everything fits into a single chunk, so we have to compute everything here
+					int NumChunksRequired                   = (int)(TotalNumLoopBytesRequired + NumBytesPerBatchAvailableForTempRegs - 1) / (int)NumBytesPerBatchAvailableForTempRegs;
+					check(NumChunksRequired > 1);
+					if (NumChunksRequired < MaxChunksPerBatch) //everything fits in a single batch
 					{
-						if (RegInc[0] == 0)
-						{
-							VVMMemSet32(DstReg[0], *SrcReg[0], NumOutputInstances);
-						}
-						else
-						{
-							VVMMemCpy(DstReg[0], SrcReg[0], sizeof(uint32) * NumOutputInstances);
-						}
-						if (RegInc[1] == 0)
-						{
-							VVMMemSet32(DstReg[1], *SrcReg[1], NumOutputInstances);
-						}
-						else
-						{
-							VVMMemCpy(DstReg[1], SrcReg[1], sizeof(uint32) * NumOutputInstances);
-						}
-						if (RegInc[2] == 0)
-						{
-							VVMMemSet32(DstReg[2], *SrcReg[2], NumOutputInstances);
-						}
-						else
-						{
-							VVMMemCpy(DstReg[2], SrcReg[2], sizeof(uint32) * NumOutputInstances);
-						}
-						if (RegInc[3] == 0)
-						{
-							VVMMemSet32(DstReg[3], *SrcReg[3], NumOutputInstances);
-						}
-						else
-						{
-							VVMMemCpy(DstReg[3], SrcReg[3], sizeof(uint32) * NumOutputInstances);
-						}
+						NumChunksPerBatch                   = NumChunksRequired;
+						//take as little memory as possible and execute it in equal sized chunks
+						MaxLoopsPerChunk                    = (TotalNumLoopsRequired + NumChunksRequired - 1) / NumChunksRequired;
 					}
-					else
+					else //not everything fits in a single batch, we have to thread this
 					{
-						for (uint32 i = 0; i < NumOutputInstances; ++i)
+						MaxLoopsPerChunk                    = (uint32)(NumBytesPerBatchAvailableForTempRegs / NumBytesRequiredPerLoop);
+						uint32 NumLoopsPerBatch             = MaxLoopsPerChunk * NumChunksPerBatch;
+						NumBatches                          = (TotalNumLoopsRequired + NumLoopsPerBatch - 1) / NumLoopsPerBatch;
+						if (GVVMMaxThreadsPerScript > 0 && NumBatches > GVVMMaxThreadsPerScript)
 						{
-							DstReg[0][i] = SrcReg[0][DstIdxReg[i] & RegInc[0]];
-							DstReg[1][i] = SrcReg[1][DstIdxReg[i] & RegInc[1]];
-							DstReg[2][i] = SrcReg[2][DstIdxReg[i] & RegInc[2]];
-							DstReg[3][i] = SrcReg[3][DstIdxReg[i] & RegInc[3]];
+							//number of batches exceed the number of threads allowed.. increase the number of chunks per batch
+							NumLoopsPerBatch                = (TotalNumLoopsRequired + GVVMMaxThreadsPerScript - 1) / GVVMMaxThreadsPerScript;
+							NumChunksPerBatch               = (int)(NumLoopsPerBatch + MaxLoopsPerChunk - 1) / MaxLoopsPerChunk;
+							NumBatches                      = GVVMMaxThreadsPerScript;
+							check(NumBatches * NumChunksPerBatch * MaxLoopsPerChunk >= TotalNumLoopsRequired);
 						}
 					}
-					InsPtr += 21;
 				}
-				break;
-				case EVectorVMOp::output_batch7:
+				else
 				{
-					uint16 DataSetIdx         = VecIndices[4];
-					int *DstIdxReg            = (int *)VecReg[0];
-					uint32 RegTypeOffset      = VVMState->DataSets[DataSetIdx].OutputRegisterTypeOffsets[InsPtr[32]];
-					uint32 **OutputBuffers    = (uint32 **)VVMState->DataSets[DataSetIdx].OutputRegisters.GetData();
-					uint32 InstanceOffset     = BatchState->ChunkLocalData.StartingOutputIdxPerDataSet[DataSetIdx];
-					uint32 NumOutputInstances = BatchState->ChunkLocalData.NumOutputPerDataSet[DataSetIdx];
-
-					{ //first 3
-						uint32 *DstReg[3]         = { OutputBuffers[RegTypeOffset + VecIndices[9] ] + InstanceOffset, 
-													  OutputBuffers[RegTypeOffset + VecIndices[10]] + InstanceOffset,
-													  OutputBuffers[RegTypeOffset + VecIndices[11]] + InstanceOffset };
-						uint32 *SrcReg[3]         = { (uint32 *)VecReg[1], 
-													  (uint32 *)VecReg[2], 
-													  (uint32 *)VecReg[3] };
-
-					
-						if (NumOutputInstances == BatchState->NumInstances)
-						{
-							if (RegInc[1] == 0)
-							{
-								VVMMemSet32(DstReg[0], *SrcReg[0], NumOutputInstances);
-							}
-							else
-							{
-								VVMMemCpy(DstReg[0], SrcReg[0], sizeof(uint32) * NumOutputInstances);
-							}
-							if (RegInc[2] == 0)
-							{
-								VVMMemSet32(DstReg[1], *SrcReg[1], NumOutputInstances);
-							} 
-							else
-							{
-								VVMMemCpy(DstReg[1], SrcReg[1], sizeof(uint32) * NumOutputInstances);
-							}
-							if (RegInc[3] == 0)
-							{
-								VVMMemSet32(DstReg[2], *SrcReg[2], NumOutputInstances);
-							}
-							else
-							{
-								VVMMemCpy(DstReg[2], SrcReg[2], sizeof(uint32) * NumOutputInstances);
-							}
-						}
-						else
-						{
-							for (uint32 i = 0; i < NumOutputInstances; ++i)
-							{
-								DstReg[0][i] = SrcReg[0][DstIdxReg[i & RegInc[0]] & RegInc[1]];
-								DstReg[1][i] = SrcReg[1][DstIdxReg[i & RegInc[0]] & RegInc[2]];
-								DstReg[2][i] = SrcReg[2][DstIdxReg[i & RegInc[0]] & RegInc[3]];
-							}
-						}
-					}
-					VVMDecodeInstructionRegisters(InsPtr + 10, VecReg, RegInc);
-					{ //next 4
-						uint32 *DstReg[4]         = { OutputBuffers[RegTypeOffset + VecIndices[12]] + InstanceOffset, 
-													  OutputBuffers[RegTypeOffset + VecIndices[13]] + InstanceOffset, 
-													  OutputBuffers[RegTypeOffset + VecIndices[14]] + InstanceOffset, 
-													  OutputBuffers[RegTypeOffset + VecIndices[15]] + InstanceOffset };
-
-						uint32 *SrcReg[4]         = { (uint32 *)VecReg[0], 
-													  (uint32 *)VecReg[1], 
-													  (uint32 *)VecReg[2], 
-													  (uint32 *)VecReg[3] };
-						if (NumOutputInstances == BatchState->NumInstances)
-						{
-							if (RegInc[0] == 0)
-							{
-								VVMMemSet32(DstReg[0], *SrcReg[0], NumOutputInstances);
-							}
-							else
-							{
-								VVMMemCpy(DstReg[0], SrcReg[0], sizeof(uint32) * NumOutputInstances);
-							}
-							if (RegInc[1] == 0)
-							{
-								VVMMemSet32(DstReg[1], *SrcReg[1], NumOutputInstances);
-							}
-							else
-							{
-								VVMMemCpy(DstReg[1], SrcReg[1], sizeof(uint32) * NumOutputInstances);
-							}
-							if (RegInc[2] == 0)
-							{
-								VVMMemSet32(DstReg[2], *SrcReg[2], NumOutputInstances);
-							}
-							else
-							{
-								VVMMemCpy(DstReg[2], SrcReg[2], sizeof(uint32) * NumOutputInstances);
-							}
-							if (RegInc[3] == 0)
-							{
-								VVMMemSet32(DstReg[3], *SrcReg[3], NumOutputInstances);
-							}
-							else
-							{
-								VVMMemCpy(DstReg[3], SrcReg[3], sizeof(uint32) * NumOutputInstances);
-							}
-						}
-						else
-						{
-							for (uint32 i = 0; i < NumOutputInstances; ++i)
-							{
-								DstReg[0][i] = SrcReg[0][DstIdxReg[i] & RegInc[0]];
-								DstReg[1][i] = SrcReg[1][DstIdxReg[i] & RegInc[1]];
-								DstReg[2][i] = SrcReg[2][DstIdxReg[i] & RegInc[2]];
-								DstReg[3][i] = SrcReg[3][DstIdxReg[i] & RegInc[3]];
-							}
-						}
-					}
-					InsPtr += 33;
+					//everything fits into a single chunk
+					NumChunksPerBatch = 1;
+					MaxLoopsPerChunk = TotalNumLoopsRequired;
 				}
-				break;
-				case EVectorVMOp::output_batch8:
-				{
-					uint16 *OutputIndices     = VecIndices + 8;
-					uint16 DataSetIdx         = VecIndices[16];
-					uint16 IdxRegIdx          = VecIndices[17];
-					int *DstIdxReg            = (int *)(BatchState->RegisterData + VVMState->NumConstBuffers + NumLoops * IdxRegIdx); //guaranteed by the optimizer to be a temp register and not a const
-					uint32 RegTypeOffset      = VVMState->DataSets[DataSetIdx].OutputRegisterTypeOffsets[InsPtr[36]];
-					uint32 **OutputBuffers    = (uint32 **)VVMState->DataSets[DataSetIdx].OutputRegisters.GetData();
-					uint32 InstanceOffset     = BatchState->ChunkLocalData.StartingOutputIdxPerDataSet[DataSetIdx];
-					uint32 NumOutputInstances = BatchState->ChunkLocalData.NumOutputPerDataSet[DataSetIdx];
-					
-					{ //first block of 4
-						uint32 *DstReg[4]         = { OutputBuffers[RegTypeOffset + OutputIndices[0]] + InstanceOffset, 
-													  OutputBuffers[RegTypeOffset + OutputIndices[1]] + InstanceOffset, 
-													  OutputBuffers[RegTypeOffset + OutputIndices[2]] + InstanceOffset, 
-													  OutputBuffers[RegTypeOffset + OutputIndices[3]] + InstanceOffset };
-
-						uint32 *SrcReg[4]         = { (uint32 *)VecReg[0], 
-													  (uint32 *)VecReg[1], 
-													  (uint32 *)VecReg[2], 
-													  (uint32 *)VecReg[3] };
-
-						if (NumOutputInstances == BatchState->NumInstances)
-						{
-							if (RegInc[0] == 0)
-							{
-								VVMMemSet32(DstReg[0], *SrcReg[0], NumOutputInstances);
-							}
-							else
-							{
-								VVMMemCpy(DstReg[0], SrcReg[0], sizeof(uint32) * NumOutputInstances);
-							}
-							if (RegInc[1] == 0)
-							{
-								VVMMemSet32(DstReg[1], *SrcReg[1], NumOutputInstances);
-							}
-							else
-							{
-								VVMMemCpy(DstReg[1], SrcReg[1], sizeof(uint32) * NumOutputInstances);
-							}
-							if (RegInc[2] == 0)
-							{
-								VVMMemSet32(DstReg[2], *SrcReg[2], NumOutputInstances);
-							}
-							else
-							{
-								VVMMemCpy(DstReg[2], SrcReg[2], sizeof(uint32) * NumOutputInstances);
-							}
-							if (RegInc[3] == 0)
-							{
-								VVMMemSet32(DstReg[3], *SrcReg[3], NumOutputInstances);
-							}
-							else
-							{
-								VVMMemCpy(DstReg[3], SrcReg[3], sizeof(uint32) * NumOutputInstances);
-							}
-						}
-						else
-						{
-							for (uint32 i = 0; i < NumOutputInstances; ++i)
-							{
-								DstReg[0][i] = SrcReg[0][DstIdxReg[i] & RegInc[0]];
-								DstReg[1][i] = SrcReg[1][DstIdxReg[i] & RegInc[1]];
-								DstReg[2][i] = SrcReg[2][DstIdxReg[i] & RegInc[2]];
-								DstReg[3][i] = SrcReg[3][DstIdxReg[i] & RegInc[3]];
-							}
-						}
-					}
-					VVMDecodeInstructionRegisters(InsPtr + 8, VecReg, RegInc);
-					{ //second block of 4
-						uint32 *DstReg[4]         = { OutputBuffers[RegTypeOffset + OutputIndices[4]] + InstanceOffset, 
-													  OutputBuffers[RegTypeOffset + OutputIndices[5]] + InstanceOffset, 
-													  OutputBuffers[RegTypeOffset + OutputIndices[6]] + InstanceOffset, 
-													  OutputBuffers[RegTypeOffset + OutputIndices[7]] + InstanceOffset };
-
-						uint32 *SrcReg[4]         = { (uint32 *)VecReg[0], 
-													  (uint32 *)VecReg[1], 
-													  (uint32 *)VecReg[2], 
-													  (uint32 *)VecReg[3] };
-						if (NumOutputInstances == BatchState->NumInstances)
-						{
-							if (RegInc[0] == 0)
-							{
-								VVMMemSet32(DstReg[0], *SrcReg[0], NumOutputInstances);
-							}
-							else
-							{
-								VVMMemCpy(DstReg[0], SrcReg[0], sizeof(uint32) * NumOutputInstances);
-							}
-							if (RegInc[1] == 0)
-							{
-								VVMMemSet32(DstReg[1], *SrcReg[1], NumOutputInstances);
-							}
-							else
-							{
-								VVMMemCpy(DstReg[1], SrcReg[1], sizeof(uint32) * NumOutputInstances);
-							}
-							if (RegInc[2] == 0)
-							{
-								VVMMemSet32(DstReg[2], *SrcReg[2], NumOutputInstances);
-							}
-							else
-							{
-								VVMMemCpy(DstReg[2], SrcReg[2], sizeof(uint32) * NumOutputInstances);
-							}
-							if (RegInc[3] == 0)
-							{
-								VVMMemSet32(DstReg[3], *SrcReg[3], NumOutputInstances);
-							}
-							else
-							{
-								VVMMemCpy(DstReg[3], SrcReg[3], sizeof(uint32) * NumOutputInstances);
-							}
-						}
-						else
-						{
-							for (uint32 i = 0; i < NumOutputInstances; ++i)
-							{
-								DstReg[0][i] = SrcReg[0][DstIdxReg[i] & RegInc[0]];
-								DstReg[1][i] = SrcReg[1][DstIdxReg[i] & RegInc[1]];
-								DstReg[2][i] = SrcReg[2][DstIdxReg[i] & RegInc[2]];
-								DstReg[3][i] = SrcReg[3][DstIdxReg[i] & RegInc[3]];
-							}
-						}
-					}
-					InsPtr += 37;
-				} break;
-				default: break;
 			}
-			VVMSer_insEndExp(SerializeState, (int)(VVMSerStartOpPtr - VVMState->Bytecode), (int)(InsPtr - VVMSerStartOpPtr));
+			PerBatchRegisterDataBytesRequired = MaxLoopsPerChunk * NumBytesRequiredPerLoop;
 		}
-		VVMSer_chunkEndExp(SerializeState);
-	}
-	VVMSer_batchEndExp(SerializeState);
 
-	int32 WasCurrentlyExecuting = FPlatformAtomics::InterlockedCompareExchange(&BatchState->CurrentlyExecuting, 0, 1);
-	check(WasCurrentlyExecuting == 1); //Sanity test to make sure that the CurrentlyExcuting flag was never changed during execution of this batch
+		size_t NumBytesRequiredPerBatch = ExecCtx->VVMState->BatchOverheadSize + PerBatchRegisterDataBytesRequired;
+		ExecCtx->Internal.NumBytesRequiredPerBatch          = (uint32)NumBytesRequiredPerBatch;
+		ExecCtx->Internal.PerBatchRegisterDataBytesRequired = (uint32)PerBatchRegisterDataBytesRequired;
+		ExecCtx->Internal.NumBatches                        = NumBatches;
+		ExecCtx->Internal.MaxChunksPerBatch                 = NumChunksPerBatch;
+		ExecCtx->Internal.MaxInstancesPerChunk              = MaxLoopsPerChunk << 2;
 
-	FPlatformAtomics::InterlockedAdd(&VVMState->NumInstancesCompleted, BatchState->NumInstances);
-}
+		ExecCtx->VVMState->ExecCtxCache.NumBytesRequiredPerBatch          = ExecCtx->Internal.NumBytesRequiredPerBatch;
+		ExecCtx->VVMState->ExecCtxCache.PerBatchRegisterDataBytesRequired = ExecCtx->Internal.PerBatchRegisterDataBytesRequired;
+		ExecCtx->VVMState->ExecCtxCache.NumBatches                        = ExecCtx->Internal.NumBatches;
+		ExecCtx->VVMState->ExecCtxCache.MaxChunksPerBatch                 = ExecCtx->Internal.MaxChunksPerBatch;
+		ExecCtx->VVMState->ExecCtxCache.MaxInstancesPerChunk              = ExecCtx->Internal.MaxInstancesPerChunk;
+
+		ExecCtx->VVMState->NumInstancesExecCached = ExecCtx->NumInstances;
+	}
+	ExecCtx->Internal.NumInstancesAssignedToBatches = 0;
+	ExecCtx->Internal.NumInstancesCompleted         = 0;
+
+	for (uint32 i = 0; i < ExecCtx->VVMState->NumOutputDataSets; ++i)
+	{
+		ExecCtx->VVMState->NumOutputPerDataSet[i] = 0;
+	}
 
-VECTORVM_API void ExecVectorVMState(FVectorVMState *VVMState, FVectorVMSerializeState *SerializeState, FVectorVMSerializeState *CmpSerializeState)
-{
 #if defined(VVM_INCLUDE_SERIALIZATION) && !defined(VVM_SERIALIZE_NO_WRITE)
 	uint64 StartTime = FPlatformTime::Cycles64();
 	if (SerializeState)
@@ -2145,49 +1493,52 @@ VECTORVM_API void ExecVectorVMState(FVectorVMState *VVMState, FVectorVMSerialize
 		SerializeState->SerializeDt = 0;
 	}
 #endif //VVM_INCLUDE_SERIALIZATION
-
-	if (VVMState->NumBatches > 1)
+	if (ExecCtx->Internal.NumBatches > 1)
 	{
 #ifdef VVM_USE_OFFLINE_THREADING
 		if (parallelJobFn)
 		{
-			for (int i = 0; i < VVMState->NumBatches; ++i)
+			for (uint32 i = 0; i < ExecCtx->Internal.NumBatches; ++i)
 			{
-				parallelJobFn(ExecVVMBatch, VVMState, i, SerializeState, CmpSerializeState);
+				parallelJobFn(ExecVVMBatch, ExecCtx, i, SerializeState, CmpSerializeState);
 			}
 		}
 		else
 		{
-			for (int i = 0; i < VVMState->NumBatches; ++i)
+			for (uint32 i = 0; i < ExecCtx->Internal.NumBatches; ++i)
 			{
-				ExecVVMBatch(VVMState, i, SerializeState, CmpSerializeState);
+				ExecVVMBatch(ExecCtx, i, SerializeState, CmpSerializeState);
 			}
 		}
 #else
 		auto ExecChunkBatch = [&](int32 BatchIdx)
 		{
-			ExecVVMBatch(VVMState, BatchIdx, SerializeState, CmpSerializeState);
+			ExecVVMBatch(ExecCtx, BatchIdx, SerializeState, CmpSerializeState);
 		};
-		ParallelFor(VVMState->NumBatches, ExecChunkBatch, true);// GbParallelVVM == 0 || !bParallel);
+		//for (uint32 i = 0; i < ExecCtx->Internal.NumBatches; ++i)
+		//{
+		//	ExecVVMBatch(ExecCtx, i, SerializeState, CmpSerializeState);
+		//}
+		ParallelFor(ExecCtx->Internal.NumBatches, ExecChunkBatch, true);// GbParallelVVM == 0 || !bParallel);
 #endif
 	}
 	else
 	{
-		ExecVVMBatch(VVMState, 0, SerializeState, CmpSerializeState);
+		ExecVVMBatch(ExecCtx, 0, SerializeState, CmpSerializeState);
 	}
 
 #ifdef VVM_USE_OFFLINE_THREADING
 	//Unreal's ParallelFor() will block the executing thread until it's finished.  That isn't guaranteed
 	//outside of UE, (ie: the debugger, the only other thing that uses this as of this writing) so block.
-	while (VVMState->NumInstancesCompleted < VVMState->TotalNumInstances)
+	while (ExecCtx->Internal.NumInstancesCompleted < ExecCtx->NumInstances)
 	{
 		FPlatformProcess::Yield();
 	}
 #endif
 
-	for (uint32 i = 0; i < VVMState->NumOutputDataSets; ++i)
+	for (uint32 i = 0; i < ExecCtx->VVMState->NumOutputDataSets; ++i)
 	{
-		VVMState->DataSets[i].DataSetAccessIndex = VVMState->NumOutputPerDataSet[i] - 1;
+		ExecCtx->DataSets[i].DataSetAccessIndex = ExecCtx->VVMState->NumOutputPerDataSet[i] - 1;
 	}
 
 #if defined(VVM_INCLUDE_SERIALIZATION) && !defined(VVM_SERIALIZE_NO_WRITE)
@@ -2205,6 +1556,7 @@ VECTORVM_API void ExecVectorVMState(FVectorVMState *VVMState, FVectorVMSerialize
 #undef VVM_ALIGN
 #undef VVM_ALIGN_4
 #undef VVM_ALIGN_16
+#undef VVM_ALIGN_32
 #undef VVM_ALIGN_64
 
 #undef VVMSet_m128Const
@@ -2213,11 +1565,10 @@ VECTORVM_API void ExecVectorVMState(FVectorVMState *VVMState, FVectorVMSerialize
 #undef VVM_m128Const
 #undef VVM_m128iConst
 
-//PRAGMA_ENABLE_OPTIMIZATION
 
 #else //NIAGARA_EXP_VM
 
-VECTORVM_API FVectorVMState *InitVectorVMState(struct FVectorVMInitData *InitData, struct FVectorVMExternalFnPerInstanceData **OutPerInstanceExtData, struct FVectorVMSerializeState *SerializeState)
+VECTORVM_API FVectorVMState *AllocVectorVMState(struct FVectorVMOptimizeContext *OptimizeCtx)
 {
 	return nullptr;
 }
@@ -2233,3 +1584,4 @@ VECTORVM_API void ExecVectorVMState(FVectorVMState *VVMState, FVectorVMSerialize
 }
 
 #endif //NIAGARA_EXP_VM
+
