VectorVM Optimizations

diff --git a/Engine/Source/Runtime/VectorVM/Private/VectorVMExperimental_Memory.inl b/Engine/Source/Runtime/VectorVM/Private/VectorVMExperimental_Memory.inl
new file mode 100644
--- /dev/null
+++ b/Engine/Source/Runtime/VectorVM/Private/VectorVMExperimental_Memory.inl
@@ -0,0 +1,101 @@
+// Copyright Epic Games, Inc. All Rights Reserved.
+
+#define VVM_NUM_PAGES		(64) //do not change, if you want to use less memory, change the page size
+
+struct FVectorVMBatchCache {
+	volatile uint64 UsageBits;
+	uint8 *         Memory;
+};
+
+static FVectorVMBatchCache GlobalVVMBatchCache;
+
+VECTORVM_API void InitVectorVM() {
+	size_t VVMPageSizeInBytes = (size_t)GVVMPageSizeInKB << 10;
+	GlobalVVMBatchCache.UsageBits  = 0;
+	GlobalVVMBatchCache.Memory = (uint8 *)FMemory::Malloc(VVMPageSizeInBytes * VVM_NUM_PAGES, 4096);
+}
+
+VECTORVM_API void FreeVectorVM() {
+	VVMDefaultFree(GlobalVVMBatchCache.Memory, __FILE__, __LINE__);
+	GlobalVVMBatchCache.Memory = NULL;
+	GlobalVVMBatchCache.UsageBits  = 0;
+}
+
+struct VVMRAIIPageHandle {
+	static const int InvalidPageHandle = 0xFF;
+	void *TEMP_ptr;
+	VVMRAIIPageHandle(int PageIdx, int NumPages) : PageIdx(PageIdx) , NumPages(NumPages) { }
+	VVMRAIIPageHandle() : PageIdx(InvalidPageHandle) { }
+	~VVMRAIIPageHandle();
+	int PageIdx;
+	int NumPages;
+};
+
+static uint8 *VVMAllocBatch(size_t NumBytes, VVMRAIIPageHandle *OutPageHandle) {
+	void *mem = FMemory::Malloc(NumBytes, 16);
+	OutPageHandle->TEMP_ptr = mem;
+	return (uint8 *)mem;
+
+	size_t VVMPageSizeInBytes = (size_t)GVVMPageSizeInKB << 10;
+	size_t NumPages = (NumBytes + VVMPageSizeInBytes - 1) / VVMPageSizeInBytes;
+	if (NumPages == 0) {
+		return nullptr;
+	}
+	uint64 PageMask = ((1ULL << (uint64)NumPages) - 1ULL);
+	
+	int SanityCount = 0;
+	do
+	{
+		const uint64 StartUsage = GlobalVVMBatchCache.UsageBits;
+		uint64 v = ~GlobalVVMBatchCache.UsageBits;
+		unsigned long c = 0;
+		if (v != 0)
+			c = FPlatformMath::CountTrailingZeros64(v);
+		
+		uint64 ShiftAcc = c;
+		while (ShiftAcc + NumPages <= 64) {
+			while (v && ((~v >> c) & PageMask)) {
+				v >>= (c + 1);
+				++ShiftAcc;
+				c = FPlatformMath::CountTrailingZeros64(v);
+				ShiftAcc += c;
+			}
+			uint64 Mask = ((1ULL << (uint64)NumPages) - 1) << (uint64)ShiftAcc;
+			uint64 NewUsage = StartUsage | Mask;
+			uint64 ExpectedUsage = FPlatformAtomics::InterlockedCompareExchange((volatile int64 *)&GlobalVVMBatchCache.UsageBits, NewUsage, StartUsage);
+			if (StartUsage == ExpectedUsage) {
+				OutPageHandle->PageIdx  = (int)ShiftAcc;
+				OutPageHandle->NumPages = (int)NumPages;
+				return GlobalVVMBatchCache.Memory + ShiftAcc * VVMPageSizeInBytes;
+			}
+			break;
+		}
+		FPlatformProcess::YieldCycles(1000);
+	} while (SanityCount++ < (1 << 30));
+	VVMDebugBreakIf(SanityCount > (1 << 30) - 1);
+	OutPageHandle->PageIdx = VVMRAIIPageHandle::InvalidPageHandle;
+	return nullptr;
+}
+
+void VVMReleaseBatch(int PageIdx, int NumPages) {
+	if (NumPages <= 0) {
+		return;
+	}
+	uint64 Mask = ~((uint64)((1ULL << (uint64)NumPages) - 1ULL) << (uint64)PageIdx);
+	for (;;) {
+		uint64 OrigUsage = GlobalVVMBatchCache.UsageBits;
+		uint64 NewUsage = GlobalVVMBatchCache.UsageBits & Mask;
+		uint64 ExpectedUsage = FPlatformAtomics::InterlockedCompareExchange((volatile int64 *)&GlobalVVMBatchCache.UsageBits, NewUsage, OrigUsage);
+		if (OrigUsage == ExpectedUsage) {
+			return;
+		}
+	}
+}
+
+VVMRAIIPageHandle::~VVMRAIIPageHandle() {
+	FMemory::Free(TEMP_ptr);
+	if (PageIdx != InvalidPageHandle)
+	{
+		VVMReleaseBatch(PageIdx, NumPages);
+	}
+}
