VectorVM Optimizations

diff --git a/Engine/Source/Runtime/VectorVM/Private/VectorVMExperimental_Serialization.inl b/Engine/Source/Runtime/VectorVM/Private/VectorVMExperimental_Serialization.inl
--- a/Engine/Source/Runtime/VectorVM/Private/VectorVMExperimental_Serialization.inl
+++ b/Engine/Source/Runtime/VectorVM/Private/VectorVMExperimental_Serialization.inl
@@ -2,8 +2,17 @@
 
 #ifdef VVM_INCLUDE_SERIALIZATION
 
+enum EVVMRegFlags
+{
+	VVMRegFlag_Int      = 1,
+	VVMRegFlag_Clean    = 32,
+	VVMRegFlag_Index    = 64,
+	VVMRegFlag_Mismatch = 128,
+};
+
+
 //prototypes for stuff shared between the original and experimental VM
-static FVectorVMSerializeInstruction *VVMSerGetNextInstruction(FVectorVMSerializeState *SerializeState, int InstructionIdx, int GlobalChunkIdx, int BatchUseCount, int OpStart, int NumOps, uint64 Dt, uint64 DtDecode);
+static FVectorVMSerializeInstruction *VVMSerGetNextInstruction(FVectorVMSerializeState *SerializeState, int InstructionIdx, int GlobalChunkIdx, int OpStart, int NumOps, uint64 Dt, uint64 DtDecode);
 
 VECTORVM_API void FreeVectorVMSerializeState(FVectorVMSerializeState *SerializeState)
 {
@@ -65,6 +74,7 @@ VECTORVM_API void FreeVectorVMSerializeState(FVectorVMSerializeState *SerializeS
 		SerializeState->FreeFn(SerializeState->Instructions[i].TempRegisterFlags, __FILE__, __LINE__);
 	}
 	SerializeState->FreeFn(SerializeState->Instructions, __FILE__, __LINE__);
+	SerializeState->FreeFn(SerializeState->ConstTableSizesInBytes, __FILE__, __LINE__);
 	SerializeState->FreeFn(SerializeState->PreExecConstData, __FILE__, __LINE__);
 	SerializeState->FreeFn(SerializeState->PostExecConstData, __FILE__, __LINE__);
 }
@@ -120,35 +130,52 @@ static uint32 VectorVMSerializeSetError_(FVectorVMSerializeState *SerializeState
 #define VVMSer_regUsed(...)
 #else //VVM_SERIALIZE_NO_WRITE
 
-#define VVMSer_instruction(Type, NumParams) if (SerializeState)                                                                         \
-                                            {                                                                                           \
-                                                for (int vi = 0; vi <= (int)(NumParams); ++vi)                                          \
-                                                {                                                                                       \
-                                                    if ((VecIndices[vi] & 0x8000) == 0)											        \
-                                                    {                                                                                   \
-                                                        SerializeState->TempRegFlags[VecIndices[vi]] = VVMRegFlag_Clean + (Type);       \
-                                                    }                                                                                   \
-                                                }                                                                                       \
-                                            }
-#define VVMSer_regUsed(RegIdx, Type)        if (SerializeState && RegIdx != 0xFFFF)                                                     \
-                                            {                                                                                           \
-                                                SerializeState->TempRegFlags[RegIdx] = VVMRegFlag_Clean + (Type);                       \
-                                            }
-
-
-#define VVMSer_chunkStartExp(SerializeState, ChunkIdx_, BatchIdx)                   \
-	FVectorVMSerializeChunk *SerializeChunk = nullptr;                              \
-	int VVMSerNumInstructionsThisChunk = 0;                                         \
-	int VVMSerGlobalChunkIdx = ExecIdx * VVMState->MaxChunksPerBatch + ChunkIdx_;   \
-	if (SerializeState)                                                             \
-	{                                                                               \
-		SerializeChunk = SerializeState->Chunks + VVMSerGlobalChunkIdx;             \
-		SerializeChunk->ChunkIdx = VVMSerGlobalChunkIdx;                            \
-		SerializeChunk->BatchIdx = BatchIdx;                                        \
-		SerializeChunk->NumInstances = NumInstancesThisChunk;                       \
-		SerializeChunk->StartInstance = StartInstanceThisChunk;                     \
-		SerializeChunk->StartThreadID = FPlatformTLS::GetCurrentThreadId();         \
-		SerializeChunk->StartClock = FPlatformTime::Cycles64();                     \
+#define VVMIsRegIdxTempReg(Idx) ((Idx) < ExecCtx->VVMState->NumTempRegisters)
+#define VVMIsRegIdxConst(Idx)   ((Idx) >= ExecCtx->VVMState->NumTempRegisters && (Idx) < ExecCtx->VVMState->NumTempRegisters + ExecCtx->VVMState->NumConstsRemapped)
+#define VVMIsRegIdxInput(Idx)   ((Idx) >= ExecCtx->VVMState->NumTempRegisters + ExecCtx->VVMState->NumConstsRemapped)
+
+#define VVMSer_instruction(Type, NumParams)		//for (int vi = 0; vi <= (int)(NumParams); ++vi) {												\
+												//	int r = (int)((uint16 *)InsPtr)[vi];														\
+												//	char c = 'R';																				\
+												//	if (r >= (int)(ExecCtx->VVMState->NumTempRegisters + ExecCtx->VVMState->NumConstBuffers)) {	\
+												//		c = 'I';																				\
+												//		r -= ExecCtx->VVMState->NumTempRegisters + ExecCtx->VVMState->NumConstBuffers;			\
+												//	} else if (r >= (int)ExecCtx->VVMState->NumTempRegisters) {									\
+												//		c = 'C';																				\
+												//		r -= ExecCtx->VVMState->NumTempRegisters;												\
+												//	}																							\
+												//	printf(",%c%d", c, r);																		\
+												//}
+//if (SerializeState)                                                                     \
+												//{                                                                                       \
+												//    for (int vi = 0; vi <= (int)(NumParams); ++vi)                                      \
+												//    {                                                                                   \
+												//        if (VVMIsRegIdxTempReg(InsPtr[vi]))                                             \
+												//        {                                                                               \
+												//            SerializeState->TempRegFlags[InsPtr[vi]] = VVMRegFlag_Clean + (Type);       \
+												//        }                                                                               \
+												//    }                                                                                   \
+												//}
+
+#define VVMSer_regUsed(RegIdx, Type)		//if (SerializeState && VVMIsRegIdxTempReg(RegIdx))                                           \
+                                            //{                                                                                           \
+                                            //    SerializeState->TempRegFlags[RegIdx] = VVMRegFlag_Clean + (Type);                       \
+                                            //}
+
+
+#define VVMSer_chunkStartExp(SerializeState, ChunkIdx_, BatchIdx_)                          \
+	FVectorVMSerializeChunk *SerializeChunk = nullptr;                                      \
+	int VVMSerNumInstructionsThisChunk = 0;                                                 \
+	int VVMSerGlobalChunkIdx = ExecIdx * ExecCtx->Internal.MaxChunksPerBatch + ChunkIdx_;   \
+	if (SerializeState)                                                                     \
+	{                                                                                       \
+		SerializeChunk = SerializeState->Chunks + VVMSerGlobalChunkIdx;                     \
+		SerializeChunk->ChunkIdx = VVMSerGlobalChunkIdx;                                    \
+		SerializeChunk->BatchIdx = BatchIdx_;                                               \
+		SerializeChunk->NumInstances = NumInstancesThisChunk;                               \
+		SerializeChunk->StartInstance = StartInstanceThisChunk;                             \
+		SerializeChunk->StartThreadID = FPlatformTLS::GetCurrentThreadId();                 \
+		SerializeChunk->StartClock = FPlatformTime::Cycles64();                             \
 	}
 
 #define VVMSer_chunkEndExp(SerializeState, ...)                                                                                 \
@@ -171,7 +198,7 @@ static uint32 VectorVMSerializeSetError_(FVectorVMSerializeState *SerializeState
 	if (SerializeState)                                                                                     \
 	{                                                                                                       \
 		uint64 VVMSerEndExecCycles = FPlatformTime::Cycles64();                                             \
-		VVMSer_serializeInstruction(SerializeState, VVMState, BatchState, VVMSerNumInstructionsThisChunk++, \
+		VVMSer_serializeInstruction(SerializeState, ExecCtx, BatchState, VVMSerNumInstructionsThisChunk++,  \
 			StartInstanceThisChunk, NumInstancesThisChunk, VVMSerGlobalChunkIdx, NumLoops,                  \
 			OpStart, NumOps, VVMSerStartCycles, VVMSerEndDecodeCycles, VVMSerEndExecCycles);                \
 		uint64 VVMSerEndSerializeCycles = FPlatformTime::Cycles64();                                        \
@@ -184,14 +211,14 @@ static void VVMSer_serializeInstruction(FVectorVMSerializeState *SerializeState,
 										int OpStart, int NumOps, uint64 StartInsCycles, uint64 EndDecodeCycles, uint64 EndInsCycles)
 {
 	
-	FVectorVMSerializeInstruction *Ins = VVMSerGetNextInstruction(SerializeState, InstructionIdx, GlobalChunkIdx, BatchState->UseCount, OpStart, NumOps, EndInsCycles - StartInsCycles, EndDecodeCycles - StartInsCycles);
+	FVectorVMSerializeInstruction *Ins = VVMSerGetNextInstruction(SerializeState, InstructionIdx, GlobalChunkIdx, OpStart, NumOps, EndInsCycles - StartInsCycles, EndDecodeCycles - StartInsCycles);
 	if (Ins != nullptr)
 	{
 		check(Ins->TempRegisters != nullptr);
 		for (uint32 i = 0; i < SerializeState->NumTempRegisters; ++i)
 		{
 			FMemory::Memcpy(Ins->TempRegisters + SerializeState->NumInstances * i + StartInstance, 
-				BatchState->RegisterData + VectorVMState->NumConstBuffers + i * NumLoops,
+				BatchState->RegisterData + SerializeState->OptimizeCtx->NumConstsRemapped + i * NumLoops,
 				sizeof(uint32) * NumInstancesThisChunk);
 			check(Ins->TempRegisters != nullptr);
 			if (GlobalChunkIdx == 0)
@@ -204,19 +231,19 @@ static void VVMSer_serializeInstruction(FVectorVMSerializeState *SerializeState,
 
 #define VVMSer_initSerializationState(VectorVMState, SerializeState, InitData, Flags) VVMSer_initSerializationState_(VectorVMState, SerializeState, InitData, (uint32)Flags)
 
-static uint32 VVMSer_initSerializationState_(FVectorVMState *VVMState, FVectorVMSerializeState *SerializeState, FVectorVMInitData *InitData, uint32 Flags)
+static uint32 VVMSer_initSerializationState_(FVectorVMState *VVMState, FVectorVMSerializeState *SerializeState, FVectorVMOptimizeContext *OptimizeContext, FVectorVMExecContext *ExecContext,  uint32 Flags)
 {
 	if (SerializeState)
 	{
 		SerializeState->Error.Flags  = 0;
 
 		SerializeState->Flags            = Flags;
-		SerializeState->NumInstances     = VVMState->TotalNumInstances;
-		SerializeState->NumTempRegisters = InitData->OptimizeContext->NumTempRegisters;
-		SerializeState->NumConstBuffers  = VVMState->NumConstBuffers;
-		SerializeState->ReallocFn        = InitData->ReallocFn ? InitData->ReallocFn : VVMDefaultRealloc;
-		SerializeState->FreeFn           = InitData->FreeFn ? InitData->FreeFn : VVMDefaultFree;
-		SerializeState->OptimizeCtx      = InitData->OptimizeContext;
+		SerializeState->NumInstances     = ExecContext->NumInstances;
+		SerializeState->NumTempRegisters = VVMState->NumTempRegisters;
+		SerializeState->ConstDataCount   = VVMState->NumConstBuffers;
+		SerializeState->ReallocFn        = VVMDefaultRealloc;
+		SerializeState->FreeFn           = VVMDefaultFree;
+		SerializeState->OptimizeCtx      = OptimizeContext;
 		if (SerializeState->OptimizeCtx) {
 			SerializeState->MaxExtFnRegisters = SerializeState->OptimizeCtx->MaxExtFnRegisters;
 			SerializeState->MaxExtFnUsed = SerializeState->OptimizeCtx->MaxExtFnUsed >= 0 ? (uint32)SerializeState->OptimizeCtx->MaxExtFnUsed : 0;
@@ -224,14 +251,13 @@ static uint32 VVMSer_initSerializationState_(FVectorVMState *VVMState, FVectorVM
 			SerializeState->MaxExtFnRegisters = 0;
 			SerializeState->MaxExtFnUsed = 0;
 		}
-		
 
 		SerializeState->NumTempRegFlags = VVMState->NumTempRegisters;
-		for (int i = 0; i < InitData->DataSets.Num(); ++i)
+		for (int i = 0; i < ExecContext->DataSets.Num(); ++i)
 		{
-			if ((uint32)InitData->DataSets[i].InputRegisters.Num() > SerializeState->NumTempRegFlags)
+			if ((uint32)ExecContext->DataSets[i].InputRegisters.Num() > SerializeState->NumTempRegFlags)
 			{
-				SerializeState->NumTempRegFlags = InitData->DataSets[i].InputRegisters.Num();
+				SerializeState->NumTempRegFlags = ExecContext->DataSets[i].InputRegisters.Num();
 			}
 		}
 
@@ -245,14 +271,15 @@ static uint32 VVMSer_initSerializationState_(FVectorVMState *VVMState, FVectorVM
 			SerializeState->TempRegFlags[i] = 0;
 		}
 
-		SerializeState->NumChunks = VVMState->NumBatches * VVMState->MaxChunksPerBatch;
+		SerializeState->NumChunks = ExecContext->Internal.NumBatches * ExecContext->Internal.MaxChunksPerBatch;
 		SerializeState->Chunks = (FVectorVMSerializeChunk *)SerializeState->ReallocFn(nullptr, sizeof(FVectorVMSerializeChunk) * SerializeState->NumChunks, __FILE__, __LINE__);
 		if (SerializeState->Chunks == nullptr)
 		{
 			return VectorVMSerializeSetError(SerializeState, VVMSerErr_OutOfMemory | VVMSerErr_Init | VVMSerErr_Fatal);
 		}
 		FMemory::Memset(SerializeState->Chunks, 0, sizeof(*SerializeState->Chunks) * SerializeState->NumChunks);
-		if (SerializeVectorVMInputDataSets(SerializeState, InitData->DataSets, InitData->ConstData, InitData->NumConstData) != 0)
+		
+		if (SerializeVectorVMInputDataSets(SerializeState, ExecContext) != 0)
 		{
 			return SerializeState->Error.Flags;
 		}
@@ -284,13 +311,14 @@ static uint32 VVMSer_initSerializationState_(FVectorVMState *VVMState, FVectorVM
 			SerializeState->NumExternalData = 0;
 		}
 
-		SerializeState->Bytecode = (unsigned char *)SerializeState->ReallocFn(nullptr, InitData->OptimizeContext->NumBytecodeBytes, __FILE__, __LINE__);
+		SerializeState->Bytecode = (unsigned char *)SerializeState->ReallocFn(nullptr, OptimizeContext->NumBytecodeBytes, __FILE__, __LINE__);
 		if (SerializeState->Bytecode == nullptr)
 		{
 			return VectorVMSerializeSetError(SerializeState, VVMSerErr_OutOfMemory | VVMSerErr_Init | VVMSerErr_Fatal);
 		}
-		FMemory::Memcpy(SerializeState->Bytecode, InitData->OptimizeContext->OutputBytecode, InitData->OptimizeContext->NumBytecodeBytes);
-		SerializeState->NumBytecodeBytes = InitData->OptimizeContext->NumBytecodeBytes;
+
+		FMemory::Memcpy(SerializeState->Bytecode, OptimizeContext->OutputBytecode, OptimizeContext->NumBytecodeBytes);
+		SerializeState->NumBytecodeBytes = OptimizeContext->NumBytecodeBytes;
 	}
 	return 0;
 }
@@ -361,7 +389,7 @@ static void VVMSer_serializeInstruction(FVectorVMSerializeState *SerializeState,
 										int InstructionIdx, int StartInstance, int NumInstancesThisChunk, int GlobalChunkIdx, 
 										int OpStart, int NumOps, uint64 StartInsCycles, uint64 EndDecodeCycles, uint64 EndInsCycles)
 {
-	FVectorVMSerializeInstruction *Ins = VVMSerGetNextInstruction(SerializeState, InstructionIdx, GlobalChunkIdx, 1, OpStart, NumOps, EndInsCycles - StartInsCycles, EndDecodeCycles - StartInsCycles);
+	FVectorVMSerializeInstruction *Ins = VVMSerGetNextInstruction(SerializeState, InstructionIdx, GlobalChunkIdx, OpStart, NumOps, EndInsCycles - StartInsCycles, EndDecodeCycles - StartInsCycles);
 	if (Ins)
 	{
 		for (uint32 i = 0; i < SerializeState->NumTempRegisters; ++i)
@@ -379,10 +407,10 @@ static void VVMSer_serializeInstruction(FVectorVMSerializeState *SerializeState,
 *** Stuff shared between the new and old VM
 *********************************************************************************************************************************************************************************************************************************/
 
-static FVectorVMSerializeInstruction *VVMSerGetNextInstruction(FVectorVMSerializeState *SerializeState, int InstructionIdx, int GlobalChunkIdx, int BatchUseCount, int OpStart, int NumOps, uint64 Dt, uint64 DtDecode)
+static FVectorVMSerializeInstruction *VVMSerGetNextInstruction(FVectorVMSerializeState *SerializeState, int InstructionIdx, int GlobalChunkIdx, int OpStart, int NumOps, uint64 Dt, uint64 DtDecode)
 {
 	FVectorVMSerializeInstruction *Ins = nullptr;
-	if (GlobalChunkIdx == 0 && BatchUseCount == 1)
+	if (GlobalChunkIdx == 0)
 	{
 		check(InstructionIdx == SerializeState->NumInstructions);
 		if (SerializeState->NumInstructions + 1 >= SerializeState->NumInstructionsAllocated)
@@ -462,7 +490,7 @@ VECTORVM_API void SerializeVectorVMWriteToFile(FVectorVMSerializeState *Serializ
 		File->Write((uint8 *)&SerializeState->ExecDt           , sizeof(uint64));
 		File->Write((uint8 *)&SerializeState->SerializeDt      , sizeof(uint64));
 		File->Write((uint8 *)&SerializeState->NumTempRegisters , sizeof(uint32));
-		File->Write((uint8 *)&SerializeState->NumConstBuffers  , sizeof(uint32));
+		File->Write((uint8 *)&SerializeState->ConstDataCount    , sizeof(uint32));
 		File->Write((uint8 *)&SerializeState->NumBytecodeBytes , sizeof(uint32));
 		File->Write((uint8 *)&SerializeState->NumInstructions  , sizeof(uint32));
 		File->Write((uint8 *)&SerializeState->NumDataSets      , sizeof(uint32));
@@ -560,9 +588,15 @@ VECTORVM_API void SerializeVectorVMWriteToFile(FVectorVMSerializeState *Serializ
 			File->Write((uint8 *)&Chunk->EndClock		, sizeof(uint64));
 		}
 
-		//write const data
-		File->Write((uint8 *)SerializeState->PreExecConstData  , sizeof(uint32) * SerializeState->NumConstBuffers);
-		File->Write((uint8 *)SerializeState->PostExecConstData , sizeof(uint32) * SerializeState->NumConstBuffers);
+		{ //write const data
+			int ConstTableSizesInBytes = 0;
+			for (uint32 i = 0; i < SerializeState->ConstDataCount; ++i) {
+				ConstTableSizesInBytes += SerializeState->ConstTableSizesInBytes[i];
+			}
+			File->Write((uint8 *)SerializeState->ConstTableSizesInBytes, sizeof(uint32) * SerializeState->ConstDataCount);
+			File->Write((uint8 *)SerializeState->PreExecConstData      , ConstTableSizesInBytes);
+			File->Write((uint8 *)SerializeState->PostExecConstData     , ConstTableSizesInBytes);
+		}
 
 		//write bytecode
 		if (SerializeState->NumBytecodeBytes > 0)
@@ -608,31 +642,43 @@ VECTORVM_API void SerializeVectorVMWriteToFile(FVectorVMSerializeState *Serializ
 #endif //VVM_SERIALIZE_NO_IO
 
 
-static uint32 SerializeConstData(FVectorVMSerializeState *SerializeState, FVectorVMConstData *ConstData, int NumConstData, uint32 **OutConstData)
+static uint32 SerializeConstData(FVectorVMSerializeState *SerializeState, const uint8 * const *ConstantTableData, const int *ConstantTableSizes, int32 ConstantTableCount, uint32 **OutConstData, uint32 **OutConstTableSizes)
 {
-	int TotalNumConstDWords = 0;
-	for (int i = 0; i < NumConstData; ++i)
+	int TotalConstNumBytes = 0;
+	for (int i = 0; i < ConstantTableCount; ++i)
 	{
-		TotalNumConstDWords += ConstData[i].NumDWords;
+		TotalConstNumBytes += ConstantTableSizes[i];
 	}
-	*OutConstData = (uint32 *)SerializeState->ReallocFn(nullptr, sizeof(uint32) * TotalNumConstDWords, __FILE__, __LINE__);
+	*OutConstData = (uint32 *)SerializeState->ReallocFn(nullptr, TotalConstNumBytes, __FILE__, __LINE__);
 	if (*OutConstData == nullptr)
 	{
 		VectorVMSerializeSetError(SerializeState, VVMSerErr_OutOfMemory | VVMSerErr_ConstData | VVMSerErr_Fatal);
 		return 0;
 	}
-	uint32 *cptr = *OutConstData;
-	for (int i = 0; i < NumConstData; ++i)
+	if (OutConstTableSizes) {
+		*OutConstTableSizes = (uint32 *)SerializeState->ReallocFn(nullptr, sizeof(uint32) * ConstantTableCount, __FILE__, __LINE__);
+		if (*OutConstTableSizes == nullptr)
+		{
+			VectorVMSerializeSetError(SerializeState, VVMSerErr_OutOfMemory | VVMSerErr_ConstData | VVMSerErr_Fatal);
+			return 0;
+		}
+		for (int i = 0; i < ConstantTableCount; ++i)
+		{
+			(*OutConstTableSizes)[i] = ConstantTableSizes[i];
+		}
+	}
+	uint8 *CPtr = (uint8 *)*OutConstData;
+	for (int i = 0; i < ConstantTableCount; ++i)
 	{
-		FMemory::Memcpy(cptr, ConstData[i].RegisterData, sizeof(uint32) * ConstData[i].NumDWords);
-		cptr += ConstData[i].NumDWords;
+		FMemory::Memcpy(CPtr, ConstantTableData[i], ConstantTableSizes[i]);
+		CPtr += ConstantTableSizes[i];
 	}
-	return TotalNumConstDWords;
+	return ConstantTableCount;
 }
 
-VECTORVM_API uint32 SerializeVectorVMInputDataSets(FVectorVMSerializeState *SerializeState, TArrayView<FDataSetMeta> InDataSets, FVectorVMConstData *InConstData, int NumConstData)
+VECTORVM_API uint32 SerializeVectorVMInputDataSets(FVectorVMSerializeState *SerializeState, TArrayView<FDataSetMeta> DataSets, const uint8 * const *ConstantTableData, const int *ConstantTableSizes, int32 ConstantTableCount)
 {
-	SerializeState->NumDataSets = InDataSets.Num();
+	SerializeState->NumDataSets = DataSets.Num();
 	SerializeState->DataSets = (FVectorVMSerializeDataSet *)SerializeState->ReallocFn(nullptr, sizeof(FVectorVMSerializeDataSet) * SerializeState->NumDataSets, __FILE__, __LINE__);
 	if (SerializeState->DataSets == nullptr)
 	{
@@ -642,7 +688,7 @@ VECTORVM_API uint32 SerializeVectorVMInputDataSets(FVectorVMSerializeState *Seri
 	for (uint32 i = 0; i < SerializeState->NumDataSets; ++i)
 	{
 		FVectorVMSerializeDataSet *DsDst = SerializeState->DataSets + i;
-		FDataSetMeta *DsSrc = &InDataSets[i];
+		FDataSetMeta *DsSrc = &DataSets[i];
 		DsDst->InputInstanceOffset      = DsSrc->InstanceOffset;
 		DsDst->InputDataSetAccessIndex  = DsSrc->DataSetAccessIndex;
 		DsDst->InputIDAcquireTag        = DsSrc->IDAcquireTag;
@@ -728,18 +774,18 @@ VECTORVM_API uint32 SerializeVectorVMInputDataSets(FVectorVMSerializeState *Seri
 			DsDst->InputSpawnedIDTableNum = 0;
 		}
 	}
-	SerializeState->NumConstBuffers = SerializeConstData(SerializeState, InConstData, NumConstData, &SerializeState->PreExecConstData);
+	SerializeState->ConstDataCount = SerializeConstData(SerializeState, ConstantTableData, ConstantTableSizes, ConstantTableCount, &SerializeState->PreExecConstData, &SerializeState->ConstTableSizesInBytes);
 	return SerializeState->Error.Flags;
 }
 
-VECTORVM_API uint32 SerializeVectorVMOutputDataSets(FVectorVMSerializeState *SerializeState, TArrayView<FDataSetMeta> InDataSets, FVectorVMConstData *InConstData, int InNumConstData)
+VECTORVM_API uint32 SerializeVectorVMOutputDataSets(FVectorVMSerializeState *SerializeState, TArrayView<FDataSetMeta> DataSets, const uint8 * const *ConstantTableData, const int *ConstantTableSizes, int32 ConstantTableCount)
 {
-	if (SerializeState->NumDataSets == InDataSets.Num()) //must have already serialized the input data sets
+	if (SerializeState->NumDataSets == DataSets.Num()) //must have already serialized the input data sets
 	{
-		for (int i = 0; i < InDataSets.Num(); ++i)
+		for (int i = 0; i < DataSets.Num(); ++i)
 		{
 			FVectorVMSerializeDataSet *DsDst = SerializeState->DataSets + i;
-			FDataSetMeta *DsSrc = &InDataSets[i];
+			FDataSetMeta *DsSrc = &DataSets[i];
 			DsDst->OutputOffset[0]			= DsSrc->OutputRegisterTypeOffsets[0];
 			DsDst->OutputOffset[1]			= DsSrc->OutputRegisterTypeOffsets[1];
 			DsDst->OutputOffset[2]			= DsSrc->OutputRegisterTypeOffsets[2];
@@ -843,11 +889,23 @@ VECTORVM_API uint32 SerializeVectorVMOutputDataSets(FVectorVMSerializeState *Ser
 	}
 	if (SerializeState->Error.Flags == 0)
 	{
-		SerializeConstData(SerializeState, InConstData, InNumConstData, &SerializeState->PostExecConstData);
+		SerializeState->ConstDataCount = SerializeConstData(SerializeState, ConstantTableData, ConstantTableSizes, ConstantTableCount, &SerializeState->PostExecConstData, nullptr);
 	}
 	return SerializeState->Error.Flags;
 }
 
+#if VECTORVM_SUPPORTS_EXPERIMENTAL
+VECTORVM_API uint32 SerializeVectorVMInputDataSets(FVectorVMSerializeState *SerializeState, FVectorVMExecContext *ExecContext)
+{
+	return SerializeVectorVMInputDataSets(SerializeState, ExecContext->DataSets, ExecContext->ConstantTableData, ExecContext->ConstantTableSizes, ExecContext->ConstantTableCount);
+}
+
+VECTORVM_API uint32 SerializeVectorVMOutputDataSets(FVectorVMSerializeState *SerializeState, FVectorVMExecContext *ExecContext)
+{
+	return SerializeVectorVMOutputDataSets(SerializeState, ExecContext->DataSets, ExecContext->ConstantTableData, ExecContext->ConstantTableSizes, ExecContext->ConstantTableCount);
+}
+#endif // #if VECTORVM_SUPPORTS_EXPERIMENTAL
+
 #else //VVM_INCLUDE_SERIALIZATION
 
 #define VVMSer_batchStart(...)
@@ -881,12 +939,20 @@ VECTORVM_API uint64 VVMSer_cmpStates(FVectorVMSerializeState *S0, FVectorVMSeria
 	return 0;
 }
 
-VECTORVM_API uint32 SerializeVectorVMInputDataSets(FVectorVMSerializeState *SerializeState, TArrayView<FDataSetMeta> DataSets, FVectorVMConstData *ConstData, int NumConstData)
+VECTORVM_API uint32  SerializeVectorVMInputDataSets  (FVectorVMSerializeState *SerializeState, struct FVectorVMExecContext *ExecContext) {
+	return 0;
+}
+
+VECTORVM_API uint32  SerializeVectorVMOutputDataSets (FVectorVMSerializeState *SerializeState, struct FVectorVMExecContext *ExecContext) {
+	return 0;
+}
+
+VECTORVM_API uint32 SerializeVectorVMInputDataSets(FVectorVMSerializeState *SerializeState, TArrayView<FDataSetMeta>, const uint8 * const *ConstantTableData, const int *ConstantTableSizes, int32 ConstantTableCount)
 {
 	return 0;
 }
 
-VECTORVM_API uint32 SerializeVectorVMOutputDataSets(FVectorVMSerializeState *SerializeState, TArrayView<FDataSetMeta> DataSets, FVectorVMConstData *ConstData, int NumConstData)
+VECTORVM_API uint32 SerializeVectorVMOutputDataSets(FVectorVMSerializeState *SerializeState, TArrayView<FDataSetMeta>, const uint8 * const *ConstantTableData, const int *ConstantTableSizes, int32 ConstantTableCount)
 {
 	return 0;
 }
diff --git a/Engine/Source/Runtime/VectorVM/Public/VectorVMCommon.h b/Engine/Source/Runtime/VectorVM/Public/VectorVMCommon.h
--- a/Engine/Source/Runtime/VectorVM/Public/VectorVMCommon.h
+++ b/Engine/Source/Runtime/VectorVM/Public/VectorVMCommon.h
@@ -27,161 +27,207 @@ enum class EVectorVMOperandLocation : uint8
 	Num
 };
 
-UENUM()
-enum class EVectorVMOp : uint8
-{
-	done,
-	add,
-	sub,
-	mul,
-	div,
-	mad,
-	lerp,
-	rcp,
-	rsq,
-	sqrt,
-	neg,
-	abs,
-	exp,
-	exp2,
-	log,
-	log2,
-	sin,
-	cos,
-	tan,
-	asin,
-	acos,
-	atan,
-	atan2,
-	ceil,
-	floor,
-	fmod,
-	frac,
-	trunc,
-	clamp,
-	min,
-	max,
-	pow,
-	round,
-	sign,
-	step,
-	random,
-	noise,
-
-	//Comparison ops.
-	cmplt,
-	cmple,
-	cmpgt,
-	cmpge,
-	cmpeq,
-	cmpneq,
-	select,
-
-// 	easein,  Pretty sure these can be replaced with just a single smoothstep implementation.
-// 	easeinout,
-
-	//Integer ops
-	addi,
-	subi,
-	muli,
-	divi,//SSE Integer division is not implemented as an intrinsic. Will have to do some manual implementation.
-	clampi,
-	mini,
-	maxi,
-	absi,
-	negi,
-	signi,
-	randomi,
-	cmplti,
-	cmplei,
-	cmpgti,
-	cmpgei,
-	cmpeqi,
-	cmpneqi,
-	bit_and,
-	bit_or,
-	bit_xor,
-	bit_not,
-	bit_lshift,
-	bit_rshift,
-
-	//"Boolean" ops. Currently handling bools as integers.
-	logic_and,
-	logic_or,
-	logic_xor,
-	logic_not,
-
-	//conversions
-	f2i,
-	i2f,
-	f2b,
-	b2f,
-	i2b,
-	b2i,
-
-	// data read/write
-	inputdata_float,
-	inputdata_int32,
-	inputdata_half,
-	inputdata_noadvance_float,
-	inputdata_noadvance_int32,
-	inputdata_noadvance_half,
-	outputdata_float,
-	outputdata_int32,
-	outputdata_half,
-	acquireindex,
-
-	external_func_call,
-
-	/** Returns the index of each instance in the current execution context. */
-	exec_index,
-
-	noise2D,
-	noise3D,
-
-	/** Utility ops for hooking into the stats system for performance analysis. */
-	enter_stat_scope,
-	exit_stat_scope,
-
-	//updates an ID in the ID table
-	update_id,
-	//acquires a new ID from the free list.
-	acquire_id,
-
-	//experimental VM only
-	fused_input1_1, /* 95    op has 1 input operand	   - binary 1          ie: register 0 is an input          (ORDER IS CRUCIAL) */
-	fused_input2_1, /* 96    op has 2 input operands   - binary 01         ie: register 0 is an input          (ORDER IS CRUCIAL) */
-	fused_input2_2, /* 97    op has 2 input operands   - binary 10         ie: register 1 is an input          (ORDER IS CRUCIAL) */
-	fused_input2_3, /* 97    op has 2 input operands   - binary 11         ie: register 1 and 2 are inputs     (ORDER IS CRUCIAL) */
-	fused_input3_1, /* 98    op has 3 input operands   - binary 001        ie: register 1 is an input          (ORDER IS CRUCIAL) */
-	fused_input3_2, /* 99    op has 3 input operands   - binary 010        ie: register 2 is an input          (ORDER IS CRUCIAL) */
-	fused_input3_4, /* 101   op has 3 input operands   - binary 100        ie: register 3 is an input          (ORDER IS CRUCIAL) */
-	fused_input3_3, /* 100   op has 3 input operands   - binary 011        ie: register 1 and 2 are inputs     (ORDER IS CRUCIAL) */
-	fused_input3_5, /* 102   op has 3 input operands   - binary 101        ie: register 3 and 1 are inputs     (ORDER IS CRUCIAL) */
-	fused_input3_6, /* 103   op has 3 input operands   - binary 110        ie: register 2 and 2 are inputs     (ORDER IS CRUCIAL) */
-	fused_input3_7, /* 104   op has 3 input operands   - binary 111        ie: register 1, 2 and 3 are inputs  (ORDER IS CRUCIAL) */
-	copy_to_output, /* 105 */
-	output_batch2,  /* 106 */
-	output_batch3,  /* 107 */
-	output_batch4,  /* 108 */
-	output_batch7,  /* 108 */
-	output_batch8,  /* 108 */
-
+//              OpCode                       Category   num params	int/float flags
+#define VVM_OP_XM_LIST																																									                                   \
+	VVM_OP_XM( done                         , Other     , 0			, VVM_INS_PARAM_FFFFFF              ) /* 0  */                                                                                                         \
+	VVM_OP_XM( add                          , Op        , 3			, VVM_INS_PARAM_FFFFFF              ) /* 1  */                                                                                                         \
+	VVM_OP_XM( sub                          , Op        , 3			, VVM_INS_PARAM_FFFFFF              ) /* 2  */                                                                                                         \
+	VVM_OP_XM( mul                          , Op        , 3			, VVM_INS_PARAM_FFFFFF              ) /* 3  */                                                                                                         \
+	VVM_OP_XM( div                          , Op        , 3			, VVM_INS_PARAM_FFFFFF              ) /* 4  */                                                                                                         \
+	VVM_OP_XM( mad                          , Op        , 4			, VVM_INS_PARAM_FFFFFF              ) /* 5  */                                                                                                         \
+	VVM_OP_XM( lerp                         , Op        , 4			, VVM_INS_PARAM_FFFFFF              ) /* 6  */                                                                                                         \
+	VVM_OP_XM( rcp                          , Op        , 2			, VVM_INS_PARAM_FFFFFF              ) /* 7  */                                                                                                         \
+	VVM_OP_XM( rsq                          , Op        , 2			, VVM_INS_PARAM_FFFFFF              ) /* 8  */                                                                                                         \
+	VVM_OP_XM( sqrt                         , Op        , 2			, VVM_INS_PARAM_FFFFFF              ) /* 9  */                                                                                                         \
+	VVM_OP_XM( neg                          , Op        , 2			, VVM_INS_PARAM_FFFFFF              ) /* 10 */                                                                                                         \
+	VVM_OP_XM( abs                          , Op        , 2			, VVM_INS_PARAM_FFFFFF              ) /* 11 */                                                                                                         \
+	VVM_OP_XM( exp                          , Op        , 2			, VVM_INS_PARAM_FFFFFF              ) /* 12 */                                                                                                         \
+	VVM_OP_XM( exp2                         , Op        , 2			, VVM_INS_PARAM_FFFFFF              ) /* 13 */                                                                                                         \
+	VVM_OP_XM( log                          , Op        , 2			, VVM_INS_PARAM_FFFFFF              ) /* 14 */                                                                                                         \
+	VVM_OP_XM( log2                         , Op        , 2			, VVM_INS_PARAM_FFFFFF              ) /* 15 */                                                                                                         \
+	VVM_OP_XM( sin                          , Op        , 2			, VVM_INS_PARAM_FFFFFF              ) /* 16 */                                                                                                         \
+	VVM_OP_XM( cos                          , Op        , 2			, VVM_INS_PARAM_FFFFFF              ) /* 17 */                                                                                                         \
+	VVM_OP_XM( tan                          , Op        , 2			, VVM_INS_PARAM_FFFFFF              ) /* 18 */                                                                                                         \
+	VVM_OP_XM( asin                         , Op        , 2			, VVM_INS_PARAM_FFFFFF              ) /* 19 */                                                                                                         \
+	VVM_OP_XM( acos                         , Op        , 2			, VVM_INS_PARAM_FFFFFF              ) /* 20 */                                                                                                         \
+	VVM_OP_XM( atan                         , Op        , 2			, VVM_INS_PARAM_FFFFFF              ) /* 21 */                                                                                                         \
+	VVM_OP_XM( atan2                        , Op        , 3			, VVM_INS_PARAM_FFFFFF              ) /* 22 */                                                                                                         \
+	VVM_OP_XM( ceil                         , Op        , 2			, VVM_INS_PARAM_FFFFFF              ) /* 23 */                                                                                                         \
+	VVM_OP_XM( floor                        , Op        , 2			, VVM_INS_PARAM_FFFFFF              ) /* 24 */                                                                                                         \
+	VVM_OP_XM( fmod                         , Op        , 3			, VVM_INS_PARAM_FFFFFF              ) /* 25 */                                                                                                         \
+	VVM_OP_XM( frac                         , Op        , 2			, VVM_INS_PARAM_FFFFFF              ) /* 26 */                                                                                                         \
+	VVM_OP_XM( trunc                        , Op        , 2			, VVM_INS_PARAM_FFFFFF              ) /* 27 */                                                                                                         \
+	VVM_OP_XM( clamp                        , Op        , 4			, VVM_INS_PARAM_FFFFFF              ) /* 28 */                                                                                                         \
+	VVM_OP_XM( min                          , Op        , 3			, VVM_INS_PARAM_FFFFFF              ) /* 29 */                                                                                                         \
+	VVM_OP_XM( max                          , Op        , 3			, VVM_INS_PARAM_FFFFFF              ) /* 30 */                                                                                                         \
+	VVM_OP_XM( pow                          , Op        , 3			, VVM_INS_PARAM_FFFFFF              ) /* 31 */                                                                                                         \
+	VVM_OP_XM( round                        , Op        , 2			, VVM_INS_PARAM_FFFFFF              ) /* 32 */                                                                                                         \
+	VVM_OP_XM( sign                         , Op        , 2			, VVM_INS_PARAM_FFFFFF              ) /* 33 */                                                                                                         \
+	VVM_OP_XM( step                         , Op        , 3			, VVM_INS_PARAM_FFFFFF              ) /* 34 */                                                                                                         \
+	VVM_OP_XM( random                       , Op        , 2			, VVM_INS_PARAM_FFFFFF              ) /* 35 */                                                                                                         \
+	VVM_OP_XM( noise                        , Op        , 0			, VVM_INS_PARAM_FFFFFF              ) /* 36 */                                                                                                         \
+	VVM_OP_XM( cmplt                        , Op        , 3			, VVM_INS_PARAM_FFFIFF              ) /* 37 */                                                                                                         \
+	VVM_OP_XM( cmple                        , Op        , 3			, VVM_INS_PARAM_FFFIFF              ) /* 38 */                                                                                                         \
+	VVM_OP_XM( cmpgt                        , Op        , 3			, VVM_INS_PARAM_FFFIFF              ) /* 39 */                                                                                                         \
+	VVM_OP_XM( cmpge                        , Op        , 3			, VVM_INS_PARAM_FFFIFF              ) /* 40 */                                                                                                         \
+	VVM_OP_XM( cmpeq                        , Op        , 3			, VVM_INS_PARAM_FFFIFF              ) /* 41 */                                                                                                         \
+	VVM_OP_XM( cmpneq                       , Op        , 3			, VVM_INS_PARAM_FFFIFF              ) /* 42 */                                                                                                         \
+	VVM_OP_XM( select                       , Op        , 4			, VVM_INS_PARAM_FFIFFF              ) /* 43 */                                                                                                         \
+	VVM_OP_XM( addi                         , Op        , 3			, VVM_INS_PARAM_FFIIII              ) /* 44 */                                                                                                         \
+	VVM_OP_XM( subi                         , Op        , 3			, VVM_INS_PARAM_FFIIII              ) /* 45 */                                                                                                         \
+	VVM_OP_XM( muli                         , Op        , 3			, VVM_INS_PARAM_FFIIII              ) /* 46 */                                                                                                         \
+	VVM_OP_XM( divi                         , Op        , 3			, VVM_INS_PARAM_FFIIII              ) /* 47 //SSE Integer division is not implemented as an intrinsic. Will have to do some manual implementation. */  \
+	VVM_OP_XM( clampi                       , Op        , 4			, VVM_INS_PARAM_FFIIII              ) /* 48 */                                                                                                         \
+	VVM_OP_XM( mini                         , Op        , 3			, VVM_INS_PARAM_FFIIII              ) /* 49 */                                                                                                         \
+	VVM_OP_XM( maxi                         , Op        , 3			, VVM_INS_PARAM_FFIIII              ) /* 50 */                                                                                                         \
+	VVM_OP_XM( absi                         , Op        , 2			, VVM_INS_PARAM_FFIIII              ) /* 51 */                                                                                                         \
+	VVM_OP_XM( negi                         , Op        , 2			, VVM_INS_PARAM_FFIIII              ) /* 52 */                                                                                                         \
+	VVM_OP_XM( signi                        , Op        , 2			, VVM_INS_PARAM_FFIIII              ) /* 53 */                                                                                                         \
+	VVM_OP_XM( randomi                      , Op        , 2			, VVM_INS_PARAM_FFIIII              ) /* 54 */                                                                                                         \
+	VVM_OP_XM( cmplti                       , Op        , 3			, VVM_INS_PARAM_FFIIII              ) /* 55 */                                                                                                         \
+	VVM_OP_XM( cmplei                       , Op        , 3			, VVM_INS_PARAM_FFIIII              ) /* 56 */                                                                                                         \
+	VVM_OP_XM( cmpgti                       , Op        , 3			, VVM_INS_PARAM_FFIIII              ) /* 57 */                                                                                                         \
+	VVM_OP_XM( cmpgei                       , Op        , 3			, VVM_INS_PARAM_FFIIII              ) /* 58 */                                                                                                         \
+	VVM_OP_XM( cmpeqi                       , Op        , 3			, VVM_INS_PARAM_FFIIII              ) /* 59 */                                                                                                         \
+	VVM_OP_XM( cmpneqi                      , Op        , 3			, VVM_INS_PARAM_FFIIII              ) /* 60 */                                                                                                         \
+	VVM_OP_XM( bit_and                      , Op        , 3			, VVM_INS_PARAM_FFIIII              ) /* 61 */                                                                                                         \
+	VVM_OP_XM( bit_or                       , Op        , 3			, VVM_INS_PARAM_FFIIII              ) /* 62 */                                                                                                         \
+	VVM_OP_XM( bit_xor                      , Op        , 3			, VVM_INS_PARAM_FFIIII              ) /* 63 */                                                                                                         \
+	VVM_OP_XM( bit_not                      , Op        , 3			, VVM_INS_PARAM_FFIIII              ) /* 64 */                                                                                                         \
+	VVM_OP_XM( bit_lshift                   , Op        , 3			, VVM_INS_PARAM_FFIIII              ) /* 65 */                                                                                                         \
+	VVM_OP_XM( bit_rshift                   , Op        , 3			, VVM_INS_PARAM_FFIIII              ) /* 66 */                                                                                                         \
+	VVM_OP_XM( logic_and                    , Op        , 3			, VVM_INS_PARAM_FFIIII              ) /* 67 */                                                                                                         \
+	VVM_OP_XM( logic_or                     , Op        , 3			, VVM_INS_PARAM_FFIIII              ) /* 68 */                                                                                                         \
+	VVM_OP_XM( logic_xor                    , Op        , 3			, VVM_INS_PARAM_FFIIII              ) /* 69 */                                                                                                         \
+	VVM_OP_XM( logic_not                    , Op        , 2			, VVM_INS_PARAM_FFIIII              ) /* 70 */                                                                                                         \
+	VVM_OP_XM( f2i                          , Op        , 2			, VVM_INS_PARAM_FFFFFI              ) /* 71 */                                                                                                         \
+	VVM_OP_XM( i2f                          , Op        , 2			, VVM_INS_PARAM_FFFFIF              ) /* 72 */                                                                                                         \
+	VVM_OP_XM( f2b                          , Op        , 2			, VVM_INS_PARAM_FFFFFI              ) /* 73 */                                                                                                         \
+	VVM_OP_XM( b2f                          , Op        , 2			, VVM_INS_PARAM_FFFFIF              ) /* 74 */                                                                                                         \
+	VVM_OP_XM( i2b                          , Op        , 2			, VVM_INS_PARAM_FFFFII              ) /* 75 */                                                                                                         \
+	VVM_OP_XM( b2i                          , Op        , 2			, VVM_INS_PARAM_FFFFII              ) /* 76 */                                                                                                         \
+	VVM_OP_XM( inputdata_float              , Input     , 3			, VVM_INS_PARAM_FFFFFF              ) /* 77 */                                                                                                         \
+	VVM_OP_XM( inputdata_int32              , Input     , 3			, VVM_INS_PARAM_FFFFFF              ) /* 78 */                                                                                                         \
+	VVM_OP_XM( inputdata_half               , Input     , 3			, VVM_INS_PARAM_FFFFFF              ) /* 79 */                                                                                                         \
+	VVM_OP_XM( inputdata_noadvance_float    , Input     , 3			, VVM_INS_PARAM_FFFFFF              ) /* 80 */                                                                                                         \
+	VVM_OP_XM( inputdata_noadvance_int32    , Input     , 3			, VVM_INS_PARAM_FFFFFF              ) /* 81 */                                                                                                         \
+	VVM_OP_XM( inputdata_noadvance_half     , Input     , 3			, VVM_INS_PARAM_FFFFFF              ) /* 82 */                                                                                                         \
+	VVM_OP_XM( outputdata_float             , Output    , 3			, VVM_INS_PARAM_FFFFFI              ) /* 83 */                                                                                                         \
+	VVM_OP_XM( outputdata_int32             , Output    , 3			, VVM_INS_PARAM_FFFIII              ) /* 84 */                                                                                                         \
+	VVM_OP_XM( outputdata_half              , Output    , 3			, VVM_INS_PARAM_FFFFFI              ) /* 85 */                                                                                                         \
+	VVM_OP_XM( acquireindex                 , IndexGen  , 2			, VVM_INS_PARAM_FFFFII              ) /* 86 */                                                                                                         \
+	VVM_OP_XM( external_func_call           , ExtFnCall , 0			, VVM_INS_PARAM_FFFFFF              ) /* 87 */                                                                                                         \
+	VVM_OP_XM( exec_index                   , Op        , 1         , VVM_INS_PARAM_FFFFFI              ) /* 88 */                                                                                                         \
+	VVM_OP_XM( noise2D                      , Other     , 0			, VVM_INS_PARAM_FFFFFF              ) /* 89 */                                                                                                         \
+	VVM_OP_XM( noise3D                      , Other     , 0			, VVM_INS_PARAM_FFFFFF              ) /* 90 */                                                                                                         \
+	VVM_OP_XM( enter_stat_scope             , Stat      , 1			, VVM_INS_PARAM_FFFFFF              ) /* 91 */                                                                                                         \
+	VVM_OP_XM( exit_stat_scope              , Stat      , 0			, VVM_INS_PARAM_FFFFFF              ) /* 92 */                                                                                                         \
+	VVM_OP_XM( update_id                    , RWBuffer  , 3			, VVM_INS_PARAM_FFFFFF              ) /* 93 */                                                                                                         \
+	VVM_OP_XM( acquire_id                   , RWBuffer  , 3			, VVM_INS_PARAM_FFFFFF              ) /* 94 */                                                                                                         \
+	VVM_OP_XM( half_to_float                , Fused     , 2			, VVM_INS_PARAM_FFFFFF              ) /* 95 */                                                                                                         \
+	/* Fused ops -- combined ops that show up frequently together in Fornite.  There are three types: */                                                                                                                   \
+	/* 1. exec_index that get immediately fed into an add or i2f */                                                                                                                                                        \
+	/* 2. ops with identical inputs */                                                                                                                                                                                     \
+	/* 3. ops where the output of one chain to the next.. ie a mul that feeds directly into a sub */                                                                                                                       \
+	VVM_OP_XM( exec_indexf                  , Op        , 2         , VVM_INS_PARAM_FFIIIF              ) /* 96 */                                                                                                         \
+	VVM_OP_XM( exec_index_addi              , Op        , 2         , VVM_INS_PARAM_FFIIII              ) /* 97 */                                                                                                         \
+	\
+	VVM_OP_XM( cmplt_select                 , Op        , 5         , VVM_INS_PARAM_FFFFFF              ) /* 98 */                                                                                                         \
+	VVM_OP_XM( cmple_select                 , Op        , 5         , VVM_INS_PARAM_FFFFFF              ) /* 99 */                                                                                                        \
+	VVM_OP_XM( cmpeq_select                 , Op        , 5         , VVM_INS_PARAM_FFFFFF              ) /* 100 */                                                                                                        \
+	VVM_OP_XM( cmplti_select                , Op        , 5         , VVM_INS_PARAM_FIIFFF              ) /* 101 */                                                                                                        \
+	VVM_OP_XM( cmplei_select                , Op        , 5         , VVM_INS_PARAM_FIIFFF              ) /* 102 */                                                                                                        \
+	VVM_OP_XM( cmpeqi_select                , Op        , 5         , VVM_INS_PARAM_FIIFFF              ) /* 103 */                                                                                                        \
+	\
+	VVM_OP_XM( cmplt_logic_and              , Op        , 4         , VVM_INS_PARAM_FFFFII              ) /* 104 */                                                                                                        \
+	VVM_OP_XM( cmple_logic_and              , Op        , 4         , VVM_INS_PARAM_FFFFII              ) /* 105 */                                                                                                        \
+	VVM_OP_XM( cmpgt_logic_and              , Op        , 4         , VVM_INS_PARAM_FFFFII              ) /* 106 */                                                                                                        \
+	VVM_OP_XM( cmpge_logic_and              , Op        , 4         , VVM_INS_PARAM_FFFFII              ) /* 107 */                                                                                                        \
+	VVM_OP_XM( cmpeq_logic_and              , Op        , 4         , VVM_INS_PARAM_FFFFII              ) /* 108 */                                                                                                        \
+	VVM_OP_XM( cmpne_logic_and              , Op        , 4         , VVM_INS_PARAM_FFFFII              ) /* 109 */                                                                                                        \
+	VVM_OP_XM( cmplti_logic_and             , Op        , 4         , VVM_INS_PARAM_FIIFII              ) /* 110 */                                                                                                        \
+	VVM_OP_XM( cmplei_logic_and             , Op        , 4         , VVM_INS_PARAM_FIIFII              ) /* 111 */                                                                                                        \
+	VVM_OP_XM( cmpgti_logic_and             , Op        , 4         , VVM_INS_PARAM_FIIFII              ) /* 112 */                                                                                                        \
+	VVM_OP_XM( cmpgei_logic_and             , Op        , 4         , VVM_INS_PARAM_FIIFIF              ) /* 113 */                                                                                                        \
+	VVM_OP_XM( cmpeqi_logic_and             , Op        , 4         , VVM_INS_PARAM_FIIFIF              ) /* 114 */                                                                                                        \
+	VVM_OP_XM( cmpnei_logic_and             , Op        , 4         , VVM_INS_PARAM_FIIFIF              ) /* 115 */                                                                                                        \
+	\
+	VVM_OP_XM( cmplt_logic_or               , Op        , 4         , VVM_INS_PARAM_FFFFFF              ) /* 116 */                                                                                                        \
+	VVM_OP_XM( cmple_logic_or               , Op        , 4         , VVM_INS_PARAM_FFFFFF              ) /* 117 */                                                                                                        \
+	VVM_OP_XM( cmpgt_logic_or               , Op        , 4         , VVM_INS_PARAM_FFFFFF              ) /* 118 */                                                                                                        \
+	VVM_OP_XM( cmpge_logic_or               , Op        , 4         , VVM_INS_PARAM_FFFFFF              ) /* 119 */                                                                                                        \
+	VVM_OP_XM( cmpeq_logic_or               , Op        , 4         , VVM_INS_PARAM_FFFFFF              ) /* 120 */                                                                                                        \
+	VVM_OP_XM( cmpne_logic_or               , Op        , 4         , VVM_INS_PARAM_FFFFFF              ) /* 121 */                                                                                                        \
+	VVM_OP_XM( cmplti_logic_or              , Op        , 4         , VVM_INS_PARAM_FIIFFF              ) /* 122 */                                                                                                        \
+	VVM_OP_XM( cmplei_logic_or              , Op        , 4         , VVM_INS_PARAM_FIIFFF              ) /* 123 */                                                                                                        \
+	VVM_OP_XM( cmpgti_logic_or              , Op        , 4         , VVM_INS_PARAM_FIIFFF              ) /* 124 */                                                                                                        \
+	VVM_OP_XM( cmpgei_logic_or              , Op        , 4         , VVM_INS_PARAM_FIIFFF              ) /* 125 */                                                                                                        \
+	VVM_OP_XM( cmpeqi_logic_or              , Op        , 4         , VVM_INS_PARAM_FIIFFF              ) /* 126 */                                                                                                        \
+	VVM_OP_XM( cmpnei_logic_or              , Op        , 4         , VVM_INS_PARAM_FIIFFF              ) /* 127 */                                                                                                        \
+	\
+	VVM_OP_XM( mad_add                      , Op        , 5         , VVM_INS_PARAM_FFFFFF              ) /* 128 */                                                                                                        \
+	VVM_OP_XM( mad_sub0                     , Op        , 5         , VVM_INS_PARAM_FFFFFF              ) /* 129 */                                                                                                        \
+	VVM_OP_XM( mad_sub1                     , Op        , 5         , VVM_INS_PARAM_FFFFFF              ) /* 130 */                                                                                                        \
+	VVM_OP_XM( mad_mul                      , Op        , 5         , VVM_INS_PARAM_FFFFFF              ) /* 131 */                                                                                                        \
+	VVM_OP_XM( mad_sqrt                     , Op        , 4         , VVM_INS_PARAM_FFFFFF              ) /* 132 */                                                                                                        \
+	VVM_OP_XM( mad_mad0                     , Op        , 6         , VVM_INS_PARAM_FFFFFF              ) /* 133 */                                                                                                        \
+	VVM_OP_XM( mad_mad1                     , Op        , 6         , VVM_INS_PARAM_FFFFFF              ) /* 134 */                                                                                                        \
+	\
+	VVM_OP_XM( mul_mad0                     , Op        , 5         , VVM_INS_PARAM_FFFFFF              ) /* 135 */                                                                                                        \
+	VVM_OP_XM( mul_mad1                     , Op        , 5         , VVM_INS_PARAM_FFFFFF              ) /* 136 */                                                                                                        \
+	VVM_OP_XM( mul_add                      , Op        , 4         , VVM_INS_PARAM_FFFFFF              ) /* 137 */                                                                                                        \
+	VVM_OP_XM( mul_sub0                     , Op        , 4         , VVM_INS_PARAM_FFFFFF              ) /* 138 */                                                                                                        \
+	VVM_OP_XM( mul_sub1                     , Op        , 4         , VVM_INS_PARAM_FFFFFF              ) /* 139 */                                                                                                        \
+	VVM_OP_XM( mul_mul                      , Op        , 4         , VVM_INS_PARAM_FFFFFF              ) /* 140 */                                                                                                        \
+	VVM_OP_XM( mul_max                      , Op        , 4         , VVM_INS_PARAM_FFFFFF              ) /* 141 */                                                                                                        \
+	VVM_OP_XM( mul_2x                       , Op        , 4         , VVM_INS_PARAM_FFFFFF              ) /* 142 */                                                                                                        \
+	\
+	VVM_OP_XM( add_mad1                     , Op        , 5         , VVM_INS_PARAM_FFFFFF              ) /* 143 */                                                                                                        \
+	VVM_OP_XM( add_add                      , Op        , 4         , VVM_INS_PARAM_FFFFFF              ) /* 144 */                                                                                                        \
+	\
+	VVM_OP_XM( sub_cmplt1                   , Op        , 4         , VVM_INS_PARAM_FFFFFI              ) /* 145 */                                                                                                        \
+	VVM_OP_XM( sub_neg                      , Op        , 3         , VVM_INS_PARAM_FFFFFF              ) /* 146 */                                                                                                        \
+	VVM_OP_XM( sub_mul                      , Op        , 4         , VVM_INS_PARAM_FFFFFF              ) /* 147 */                                                                                                        \
+	\
+	VVM_OP_XM( div_mad0                     , Op        , 5         , VVM_INS_PARAM_FFFFFF              ) /* 148 */                                                                                                        \
+	VVM_OP_XM( div_f2i                      , Op        , 3         , VVM_INS_PARAM_FFFFFI              ) /* 149 */                                                                                                        \
+	VVM_OP_XM( div_mul                      , Op        , 4         , VVM_INS_PARAM_FFFFFF              ) /* 150 */                                                                                                        \
+	\
+	VVM_OP_XM( muli_addi                    , Op        , 4         , VVM_INS_PARAM_IIIIII              ) /* 151 */                                                                                                        \
+	VVM_OP_XM( addi_bit_rshift				, Op        , 4         , VVM_INS_PARAM_IIIIII              ) /* 152 */                                                                                                        \
+	VVM_OP_XM( addi_muli                    , Op        , 4         , VVM_INS_PARAM_IIIIII              ) /* 153 */                                                                                                        \
+	\
+	VVM_OP_XM( b2i_2x                       , Op        , 3         , VVM_INS_PARAM_IIIIII              ) /* 154 */                                                                                                        \
+	\
+	VVM_OP_XM( i2f_div0                     , Op        , 3         , VVM_INS_PARAM_FFFFFF              ) /* 155 */                                                                                                        \
+	VVM_OP_XM( i2f_div1                     , Op        , 3         , VVM_INS_PARAM_FFFFFF              ) /* 156 */                                                                                                        \
+	VVM_OP_XM( i2f_mul                      , Op        , 3         , VVM_INS_PARAM_FFFFFF              ) /* 157 */                                                                                                        \
+	VVM_OP_XM( i2f_mad0                     , Op        , 4         , VVM_INS_PARAM_FFFFFF              ) /* 158 */                                                                                                        \
+	VVM_OP_XM( i2f_mad1                     , Op        , 4         , VVM_INS_PARAM_FFFFFF              ) /* 159 */                                                                                                        \
+	\
+	VVM_OP_XM( f2i_select1                  , Op        , 4         , VVM_INS_PARAM_IIFIII              ) /* 160 */                                                                                                        \
+	VVM_OP_XM( f2i_maxi                     , Op        , 3         , VVM_INS_PARAM_IIIFII              ) /* 161 */                                                                                                        \
+	VVM_OP_XM( f2i_addi                     , Op        , 3         , VVM_INS_PARAM_IIIFII              ) /* 162 */                                                                                                        \
+	\
+	VVM_OP_XM( fmod_add                     , Op        , 3         , VVM_INS_PARAM_FFFFFF              ) /* 163 */                                                                                                        \
+	VVM_OP_XM( bit_and_i2f                  , Op        , 3         , VVM_INS_PARAM_FFFFFF              ) /* 164 */                                                                                                        \
+	VVM_OP_XM( bit_rshift_bit_and           , Op        , 3         , VVM_INS_PARAM_FFFFFF              ) /* 165 */                                                                                                        \
+	VVM_OP_XM( neg_cmplt                    , Op        , 3         , VVM_INS_PARAM_FFFFFF              ) /* 166 */                                                                                                        \
+	VVM_OP_XM( bit_or_muli                  , Op        , 3         , VVM_INS_PARAM_FFFFFF              ) /* 167 */                                                                                                        \
+	VVM_OP_XM( bit_lshift_bit_or            , Op        , 3         , VVM_INS_PARAM_FFFFFF              ) /* 168 */                                                                                                        \
+	VVM_OP_XM( random_add                   , Op        , 3         , VVM_INS_PARAM_FFFFFF              ) /* 169 */                                                                                                        \
+	VVM_OP_XM( random_2x                    , Op        , 3         , VVM_INS_PARAM_FFFFFF              ) /* 170 */                                                                                                        \
+	VVM_OP_XM( max_f2i                      , Op        , 3         , VVM_INS_PARAM_FFFFFF              ) /* 171 */                                                                                                        \
+	VVM_OP_XM( select_mul                   , Op        , 5         , VVM_INS_PARAM_FFFFFF              ) /* 172 */                                                                                                        \
+	VVM_OP_XM( select_add                   , Op        , 5         , VVM_INS_PARAM_FFFFFF              ) /* 173 */                                                                                                        \
+	VVM_OP_XM( sin_cos                      , Op        , 3         , VVM_INS_PARAM_FFFFFF              ) /* 174 */
+
+#define VVM_OP_XM(n, ...) n,
+enum class EVectorVMOp : uint8 {
+	VVM_OP_XM_LIST
 	NumOpcodes
 };
-	
-enum class EVectorVMOpCategory : uint8 {
-	Input,
-	Output,
-	Op,
-	ExtFnCall,
-	IndexGen,
-	ExecIndex,
-	RWBuffer,
-	Stat,
-	Fused,
-	Other
-};
+#undef VVM_OP_XM
 
 #if STATS
 struct FVMCycleCounter
diff --git a/Engine/Source/Runtime/VectorVM/Public/VectorVMExperimental.h b/Engine/Source/Runtime/VectorVM/Public/VectorVMExperimental.h
--- a/Engine/Source/Runtime/VectorVM/Public/VectorVMExperimental.h
+++ b/Engine/Source/Runtime/VectorVM/Public/VectorVMExperimental.h
@@ -1,21 +1,119 @@
 // Copyright Epic Games, Inc. All Rights Reserved.
 #pragma once
 
+#define VECTORVM_SUPPORTS_AVX 0
+#define VECTORVM_SUPPORTS_COMPUTED_GOTO 0
+
 //only to be included by VectorVM.h
 
 typedef void * (VectorVMReallocFn)              (void *Ptr, size_t NumBytes, const char *Filename, int LineNumber);
 typedef void   (VectorVMFreeFn)                 (void *Ptr, const char *Filename, int LineNumber);
 struct FDataSetMeta;
 
-//prototypes for serialization that are required whether or not serialization is enabled
-#if VECTORVM_SUPPORTS_EXPERIMENTAL || defined(VVM_INCLUDE_SERIALIZATION)
-struct FVectorVMConstData
+
+#define VVM_INS_PARAM_FFFFFF	0b000000
+#define VVM_INS_PARAM_FFFFFI	0b000001
+#define VVM_INS_PARAM_FFFFIF	0b000010
+#define VVM_INS_PARAM_FFFFII	0b000011
+#define VVM_INS_PARAM_FFFIFF	0b000100
+#define VVM_INS_PARAM_FFFIFI	0b000101
+#define VVM_INS_PARAM_FFFIIF	0b000110
+#define VVM_INS_PARAM_FFFIII	0b000111
+#define VVM_INS_PARAM_FFIFFF	0b001000
+#define VVM_INS_PARAM_FFIFFI	0b001001
+#define VVM_INS_PARAM_FFIFIF	0b001010
+#define VVM_INS_PARAM_FFIFII	0b001011
+#define VVM_INS_PARAM_FFIIFF	0b001100
+#define VVM_INS_PARAM_FFIIFI	0b001101
+#define VVM_INS_PARAM_FFIIIF	0b001110
+#define VVM_INS_PARAM_FFIIII	0b001111
+#define VVM_INS_PARAM_FIFFFF	0b010000
+#define VVM_INS_PARAM_FIFFFI	0b010001
+#define VVM_INS_PARAM_FIFFIF	0b010010
+#define VVM_INS_PARAM_FIFFII	0b010011
+#define VVM_INS_PARAM_FIFIFF	0b010100
+#define VVM_INS_PARAM_FIFIFI	0b010101
+#define VVM_INS_PARAM_FIFIIF	0b010110
+#define VVM_INS_PARAM_FIFIII	0b010111
+#define VVM_INS_PARAM_FIIFFF	0b011000
+#define VVM_INS_PARAM_FIIFFI	0b011001
+#define VVM_INS_PARAM_FIIFIF	0b011010
+#define VVM_INS_PARAM_FIIFII	0b011011
+#define VVM_INS_PARAM_FIIIFF	0b011100
+#define VVM_INS_PARAM_FIIIFI	0b011101
+#define VVM_INS_PARAM_FIIIIF	0b011110
+#define VVM_INS_PARAM_FIIIII	0b011111
+#define VVM_INS_PARAM_IFFFFF	0b100000
+#define VVM_INS_PARAM_IFFFFI	0b100001
+#define VVM_INS_PARAM_IFFFIF	0b100010
+#define VVM_INS_PARAM_IFFFII	0b100011
+#define VVM_INS_PARAM_IFFIFF	0b100100
+#define VVM_INS_PARAM_IFFIFI	0b100101
+#define VVM_INS_PARAM_IFFIIF	0b100110
+#define VVM_INS_PARAM_IFFIII	0b100111
+#define VVM_INS_PARAM_IFIFFF	0b101000
+#define VVM_INS_PARAM_IFIFFI	0b101001
+#define VVM_INS_PARAM_IFIFIF	0b101010
+#define VVM_INS_PARAM_IFIFII	0b101011
+#define VVM_INS_PARAM_IFIIFF	0b101100
+#define VVM_INS_PARAM_IFIIFI	0b101101
+#define VVM_INS_PARAM_IFIIIF	0b101110
+#define VVM_INS_PARAM_IFIIII	0b101111
+#define VVM_INS_PARAM_IIFFFF	0b110000
+#define VVM_INS_PARAM_IIFFFI	0b110001
+#define VVM_INS_PARAM_IIFFIF	0b110010
+#define VVM_INS_PARAM_IIFFII	0b110011
+#define VVM_INS_PARAM_IIFIFF	0b110100
+#define VVM_INS_PARAM_IIFIFI	0b110101
+#define VVM_INS_PARAM_IIFIIF	0b110110
+#define VVM_INS_PARAM_IIFIII	0b110111
+#define VVM_INS_PARAM_IIIFFF	0b111000
+#define VVM_INS_PARAM_IIIFFI	0b111001
+#define VVM_INS_PARAM_IIIFIF	0b111010
+#define VVM_INS_PARAM_IIIFII	0b111011
+#define VVM_INS_PARAM_IIIIFF	0b111100
+#define VVM_INS_PARAM_IIIIFI	0b111101
+#define VVM_INS_PARAM_IIIIIF	0b111110
+#define VVM_INS_PARAM_IIIIII	0b111111
+
+
+#define VVM_OP_CAT_XM_LIST    \
+	VVM_OP_CAT_XM(Input)      \
+	VVM_OP_CAT_XM(Output)     \
+	VVM_OP_CAT_XM(Op)         \
+	VVM_OP_CAT_XM(ExtFnCall)  \
+	VVM_OP_CAT_XM(IndexGen)   \
+	VVM_OP_CAT_XM(RWBuffer)   \
+	VVM_OP_CAT_XM(Stat)       \
+	VVM_OP_CAT_XM(Fused)      \
+	VVM_OP_CAT_XM(Other)
+
+#define VVM_OP_CAT_XM(Category, ...) Category,
+enum class EVectorVMOpCategory : uint8 {
+	VVM_OP_CAT_XM_LIST
+	MAX
+};
+#undef VVM_OP_CAT_XM
+
+static const EVectorVMOpCategory VVM_OP_CATEGORIES[] = {
+#define VVM_OP_XM(op, cat, ...) EVectorVMOpCategory::cat,
+	VVM_OP_XM_LIST
+#undef VVM_OP_XM
+};
+
+enum EVectorVMFlags
 {
-	void * RegisterData;
-	uint32 NumDWords;
-	uint32 StartingOffset;
+	VVMFlag_OptSaveIntermediateState = 1 << 0,
+	VVMFlag_OptOmitStats             = 1 << 1,
+	VVMFlag_LargeScript              = 1 << 2,   //if true register indices are 16 bit, otherwise they're 8 bit
+	VVMFlag_HasRandInstruction       = 1 << 3,
+	VVMFlag_DataMapCacheSetup        = 1 << 4,
+	VVMFlag_SupportsAVX              = 1 << 5,
 };
 
+//prototypes for serialization that are required whether or not serialization is enabled
+#if VECTORVM_SUPPORTS_EXPERIMENTAL || defined(VVM_INCLUDE_SERIALIZATION)
+
 typedef uint32 (VectorVMOptimizeErrorCallback)  (struct FVectorVMOptimizeContext *OptimizeContext, uint32 ErrorFlags);	//return new error flags
 typedef uint32 (VectorVMSerializeErrorCallback) (struct FVectorVMSerializeState *SerializeState, uint32 ErrorFlags);
 
@@ -123,7 +221,7 @@ struct FVectorVMSerializeState
 	uint32                              NumInstances;
 	uint32                              NumTempRegisters;
 	uint32                              NumTempRegFlags;  //max of NumTempRegisters and Num Input Registers in each dataset
-	uint32                              NumConstBuffers;
+	uint32                              ConstDataCount;
 
 	uint32                              Flags;
 
@@ -145,9 +243,12 @@ struct FVectorVMSerializeState
 
 	FVectorVMSerializeDataSet *         DataSets;
 	uint32                              NumDataSets;
+
+	uint32 *                            ConstTableSizesInBytes;
 	uint32 *                            PreExecConstData;
 	uint32 *                            PostExecConstData;
 
+
 	uint32                              NumChunks;
 	FVectorVMSerializeChunk *           Chunks;
 
@@ -191,12 +292,6 @@ struct FVectorVMExtFunctionData
 };
 
 //Optimization
-enum EVectorVMOptimizeFlags
-{
-	VVMOptFlag_SaveIntermediateState = 1 << 0,
-	VVMOptFlag_OmitStats             = 1 << 1
-};
-
 struct FVectorVMOptimizeInstruction
 {
 	EVectorVMOp         OpCode;
@@ -204,6 +299,7 @@ struct FVectorVMOptimizeInstruction
 	uint32              PtrOffsetInOrigBytecode;
 	uint32              PtrOffsetInOptimizedBytecode;
 	int                 Index; //initial index.  Instructions are moved around and removed and dependency chains are created based on index, so we need to store this.
+	int					InsMergedIdx; //if not -1, then the instruction index that this is merged with.  Instructions with a set InsMergedIdx are not written to the final bytecode
 	union
 	{
 		struct
@@ -211,22 +307,18 @@ struct FVectorVMOptimizeInstruction
 			uint16 DstRegPtrOffset;
 			uint16 DataSetIdx;
 			uint16 InputIdx;
-			uint16 FuseCount;
-			int    FirstInsInsertIdx;
 		} Input;
 		struct
 		{
 			uint16 RegPtrOffset;
 			uint16 DataSetIdx;
 			uint16 DstRegIdx;
-			int    CopyFromInputInsIdx;
 		} Output;
 		struct
 		{
 			uint32 RegPtrOffset;
 			uint16 NumInputs;
 			uint16 NumOutputs;
-			uint8  InputFuseBits;
 		} Op;
 		struct
 		{
@@ -241,10 +333,6 @@ struct FVectorVMOptimizeInstruction
 			uint16 NumOutputs;
 		} ExtFnCall;
 		struct
-		{
-			uint32 RegPtrOffset;
-		} ExecIndex;
-		struct
 		{
 			uint32 RegPtrOffset;
 			uint16 DataSetIdx;
@@ -281,21 +369,29 @@ struct FVectorVMOptimizeContext
 {
 	uint8 *                               OutputBytecode;
 	uint16 *                              ConstRemap[2];
+	uint16 *                              InputRemapTable;
+	uint16 *                              InputDataSetOffsets;
 	FVectorVMExtFunctionData *            ExtFnTable;
+
 	uint32                                NumBytecodeBytes;
 	uint32                                NumOutputDataSets;
 	uint16                                NumConstsAlloced;    //upper bound to alloc
-	uint16                                NumConstsRemapped;
 	uint32                                NumTempRegisters;
+	uint16                                NumConstsRemapped;
+	uint16                                NumInputsRemapped;
+	uint16                                NumNoAdvanceInputs;
+	uint16                                NumInputDataSets;
 	uint32                                NumExtFns;
 	uint32                                MaxExtFnRegisters;
 	uint32                                NumDummyRegsReq;     //External function "null" registers
 	int32                                 MaxExtFnUsed;
+	uint32                                Flags;
 
 	struct
 	{
 		VectorVMReallocFn *               ReallocFn;
 		VectorVMFreeFn *                  FreeFn;
+		const char *                      ScriptName;
 	} Init;                                                    //Set this stuff when calling Optimize()
 
 	struct
@@ -308,9 +404,9 @@ struct FVectorVMOptimizeContext
 	struct
 	{
 		FVectorVMOptimizeInstruction *    Instructions;
+		uint8 *                           RegisterUsageType;
 		uint16 *                          RegisterUsageBuffer;
 		uint16 *                          SSARegisterUsageBuffer;
-		int32 *                           InputRegisterFuseBuffer;
 		uint32                            NumBytecodeBytes;
 		uint32                            NumInstructions;
 		uint32                            NumRegistersUsed;
@@ -336,25 +432,6 @@ struct FVectorVMExternalFnPerInstanceData
 	uint16                            NumOutputs;
 };
 
-struct FVectorVMInitData
-{
-	struct FVectorVMState *					ExistingVectorVMState;
-	const FVectorVMOptimizeContext *        OptimizeContext;
-	TArrayView<FDataSetMeta>                DataSets;
-	TArrayView<const FVMExternalFunction *> ExtFunctionTable;
-	
-	int                                     NumInstances;
-
-	int                                     NumConstData;
-	FVectorVMConstData *                    ConstData;
-
-	void **                                 UserPtrTable;
-	int                                     NumUserPtrTable;
-
-	VectorVMReallocFn *                     ReallocFn;
-	VectorVMFreeFn *                        FreeFn;
-};
-
 struct FVectorVMBatchState
 {
 	MS_ALIGN(16) FVecReg *    RegisterData GCC_ALIGN(16);
@@ -362,77 +439,108 @@ struct FVectorVMBatchState
 		uint32 *              StartingOutputIdxPerDataSet;
 		uint32 *              NumOutputPerDataSet;
 		struct {
-			FVecReg **        RegData;
-			uint32 *          RegInc;
+			uint32 **         RegData;
+			uint8 *           RegInc;
 			FVecReg *         DummyRegs;
 		} ExtFnDecodedReg;
 		int32 *               RandCounters; //used for external functions only.
 	} ChunkLocalData;
-	void *                    MallocedMemPtr; //needed for alignment purposes
-
-	volatile int32			  CurrentlyExecuting;
+	
+	uint8 **                  RegPtrTable;    //not aligned because input pointers could be offest from the DataSetInfo.InstanceOffset, so we must assume every op is unaligned
+	uint8 *                   RegIncTable;
 
 	int                       StartInstance;
 	int                       NumInstances;
-
-	VectorRegister4i          RandState[5]; //xorwor state for random/randomi instructions.  DIs use RandomStream.
-	VectorRegister4i          RandCounters;
-
+	
+	union {
+		struct {
+			VectorRegister4i          State[5]; //xorwor state for random/randomi instructions.  DIs use RandomStream.
+			VectorRegister4i          Counters;
+		};
+#		if VECTORVM_SUPPORTS_AVX
+		struct {
+			__m256i                   State[5];
+			__m256i                   Counters;
+		} AVX;
+#		endif
+	} RandState;
 	FRandomStream             RandStream;
-	int                       UseCount;
 };
 
 static_assert((sizeof(FVectorVMBatchState) & 0xF) == 0, "FVectorVMBatchState must be 16 byte aligned");
 
 struct FVectorVMState
 {
-	uint8 *                     Bytecode;
-	FVectorVMExtFunctionData *  ExtFunctionTable;
-	void **                     UserPtrTable;
-	volatile int32 *            NumOutputPerDataSet;
-	FVectorVMBatchState *       BatchStates;
-	size_t                      NumBytesMalloced;
+	uint8 *                             Bytecode;
+	uint32                              NumBytecodeBytes;
+
+	FVecReg *                           ConstantBuffers;        //the last OptimizeCtx->NumNoAdvanceInputs are no advance inputs that are copied in the table setup
+	FVectorVMExtFunctionData *          ExtFunctionTable;
+	volatile int32 *                    NumOutputPerDataSet;
+
+	uint16 *                            ConstRemapTable;
+	uint16 *                            InputRemapTable;
+	uint16 *                            InputDataSetOffsets;
+
+	uint8 *                             ConstMapCacheIdx;       //these don't get filled out until Exec() is called because they can't be filled out until the state
+	uint16 *                            ConstMapCacheSrc;       //of const and input buffers from Niagara is unknown until exec() is called.
+	uint8 *                             InputMapCacheIdx;
+	uint16 *                            InputMapCacheSrc;
+	int32                               NumInstancesExecCached;
+
+	uint32                              Flags;
 	
-	TArrayView<FDataSetMeta>    DataSets;
-
-	uint32                      NumTempRegisters;
-	uint32                      NumConstBuffers;
-	uint32                      NumBytecodeBytes;
-	uint32                      NumOutputDataSets;   //computed in the optimizer
-	uint32                      NumExtFunctions;
-	uint32                      MaxExtFnRegisters;
-	uint32                      NumDummyRegsReq;
-	uint32                      NumUserPtrTable;
-	int32                       TotalNumInstances;
-	volatile int32              NumInstancesAssignedToBatches;
-	volatile int32              NumInstancesCompleted;
-
-	size_t                      NumBytesRequiredPerBatch;
-	size_t                      PerBatchRegisterDataBytesRequired;
-	size_t                      PerBatchChunkLocalDataOutputIdxBytesRequired;
-	size_t                      PerBatchChunkLocalNumOutputBytesRequired;
-	size_t                      PerBatchChunkLocalNumExtFnDecodeRegisterBytesRequired;
-
-	int                         NumBatches;
-	int                         MaxInstancesPerChunk;
-	int                         MaxChunksPerBatch;
-
-	VectorVMReallocFn *         ReallocFn;
-	VectorVMFreeFn *            FreeFn;
+	uint32                              NumExtFunctions;
+	uint32                              MaxExtFnRegisters;
+	
+	uint32                              NumTempRegisters;
+	uint32                              NumConstBuffers;
+	uint32                              NumInputBuffers;
+	uint32                              NumInputDataSets;
+	uint32                              NumOutputDataSets;
+	uint32                              NumDummyRegsRequired;
+
+	//batch stuff
+	uint32                              BatchOverheadSize;
+	uint32                              ChunkLocalDataOutputIdxNumBytes;
+	uint32                              ChunkLocalNumOutputNumBytes;
+
+	struct {
+		uint32          NumBytesRequiredPerBatch;
+		uint32          PerBatchRegisterDataBytesRequired;
+		uint32          NumBatches;
+		uint32          MaxChunksPerBatch;
+		uint32          MaxInstancesPerChunk;
+	} ExecCtxCache;
+};
 
+struct FVectorVMExecContext {
 	struct
 	{
-		int                     LineNum;
-		uint32                  Flags;
-	} Error;
+		uint32          NumBytesRequiredPerBatch;
+		uint32          PerBatchRegisterDataBytesRequired;
+		uint32          NumBatches;
+		uint32          MaxChunksPerBatch;
+		uint32          MaxInstancesPerChunk;
+		volatile int32  NumInstancesAssignedToBatches;
+		volatile int32  NumInstancesCompleted;
+	} Internal;
+
+	FVectorVMState *                        VVMState;
+	TArrayView<FDataSetMeta>                DataSets;
+	TArrayView<const FVMExternalFunction *> ExtFunctionTable;
+	TArrayView<void*>                       UserPtrTable;
+	int32                                   NumInstances;
+	const uint8 * const *                   ConstantTableData;
+	const int *                             ConstantTableSizes;
+	int32                                   ConstantTableCount;
 };
 
 class FVectorVMExternalFunctionContextExperimental
 {
 public:
-	MS_ALIGN(16) uint32** RegisterData GCC_ALIGN(16);
-	uint16* RawVecIndices; //undecoded, for compatbility with the previous VM
-	uint32* RegInc;
+	MS_ALIGN(32) uint32**    RegisterData GCC_ALIGN(32);
+	uint8 *                  RegInc;
 
 	int                      RegReadCount;
 	int                      NumRegisters;
@@ -445,33 +553,40 @@ public:
 	void** UserPtrTable;
 	int                      NumUserPtrs;
 
-	FRandomStream* RandStream;
-	int32* RandCounters;
+	FRandomStream*           RandStream;
+	int32 **                 RandCounters;
 	TArrayView<FDataSetMeta> DataSets;
 
-	FORCEINLINE int32                                  GetStartInstance() const { return StartInstance; }
-	FORCEINLINE int32                                  GetNumInstances() const { return NumInstances; }
-	FORCEINLINE int32* GetRandCounters() { return RandCounters; }
-	FORCEINLINE FRandomStream& GetRandStream() { return *RandStream; }
-	FORCEINLINE void* GetUserPtrTable(int32 UserPtrIdx) { check(UserPtrIdx < NumUserPtrs);  return UserPtrTable[UserPtrIdx]; }
-	template<uint32 InstancesPerOp> FORCEINLINE int32  GetNumLoops() const { static_assert(InstancesPerOp == 4); return NumLoops; };
+	FORCEINLINE int32                                  GetStartInstance() const          { return StartInstance; }
+	FORCEINLINE int32                                  GetNumInstances() const           { return NumInstances; }
+	FORCEINLINE int32 *                                GetRandCounters()                 { 
+		                                                                                     if (!*RandCounters) {
+																								 *RandCounters = (int *)FMemory::MallocZeroed(sizeof(int32) * NumInstances);
+	                                                                                         }
+		                                                                                     return *RandCounters; 
+	                                                                                     }
+	FORCEINLINE FRandomStream &                        GetRandStream()                   { return *RandStream; }
+	FORCEINLINE void *                                 GetUserPtrTable(int32 UserPtrIdx) { check(UserPtrIdx < NumUserPtrs);  return UserPtrTable[UserPtrIdx]; }
+	template<uint32 InstancesPerOp> FORCEINLINE int32  GetNumLoops() const               { static_assert(InstancesPerOp == 4); return NumLoops; };
 
 	FORCEINLINE float* GetNextRegister(int32* OutAdvanceOffset, int32* OutVecIndex)
 	{
 		check(RegReadCount < NumRegisters);
-		*OutAdvanceOffset = RegInc[RegReadCount] & 1;
-		*OutVecIndex = RawVecIndices[RegReadCount];
+		*OutAdvanceOffset = RegInc[RegReadCount];
 		return (float*)RegisterData[RegReadCount++];
 	}
 };
 
 //API FUNCTIONS
 
+//initialization
+VECTORVM_API void InitVectorVM();
+VECTORVM_API void FreeVectorVM();
+
 //normal functions
-VECTORVM_API FVectorVMState * InitVectorVMState     (FVectorVMInitData *InitData, FVectorVMExternalFnPerInstanceData **OutPerInstanceExtData, FVectorVMSerializeState *SerializeState);
+VECTORVM_API FVectorVMState * AllocVectorVMState    (FVectorVMOptimizeContext *OptimizeCtx);
 VECTORVM_API void             FreeVectorVMState     (FVectorVMState *VectorVMState);
-VECTORVM_API void             ExecVectorVMState     (FVectorVMState *VectorVMState, FVectorVMSerializeState *SerializeState, FVectorVMSerializeState *CmpSerializeState);
-VECTORVM_API int              GetNumOutputInstances (FVectorVMState *VectorVMState, int DataSetIdx);
+VECTORVM_API void             ExecVectorVMState     (FVectorVMExecContext *ExecCtx, FVectorVMSerializeState *SerializeState, FVectorVMSerializeState *CmpSerializeState);
 
 //optimize functions
 VECTORVM_API uint32  OptimizeVectorVMScript                (const uint8 *Bytecode, int BytecodeLen, FVectorVMExtFunctionData *ExtFnIOData, int NumExtFns, FVectorVMOptimizeContext *OptContext, uint32 Flags); //OutContext must be zeroed except the Init struct
@@ -480,8 +595,9 @@ VECTORVM_API void    FreezeVectorVMOptimizeContext         (const FVectorVMOptim
 VECTORVM_API void    ReinterpretVectorVMOptimizeContextData(TConstArrayView<uint8> ContextData, FVectorVMOptimizeContext& Context);
 
 //serialize functions
-VECTORVM_API uint32  SerializeVectorVMInputDataSets  (FVectorVMSerializeState *SerializeState, TArrayView<FDataSetMeta> DataSets, FVectorVMConstData *ConstData, int NumConstData); //only use when not calling InitVectorVMState()
-VECTORVM_API uint32  SerializeVectorVMOutputDataSets (FVectorVMSerializeState *SerializeState, TArrayView<FDataSetMeta> DataSets, FVectorVMConstData *ConstData, int NumConstData);
+VECTORVM_API uint32  SerializeVectorVMInputDataSets  (FVectorVMSerializeState *SerializeState, struct FVectorVMExecContext *ExecContext);
+VECTORVM_API uint32  SerializeVectorVMOutputDataSets (FVectorVMSerializeState *SerializeState, struct FVectorVMExecContext *ExecContext);
+
 VECTORVM_API void    SerializeVectorVMWriteToFile    (FVectorVMSerializeState *SerializeState, uint8 WhichStateWritten, const wchar_t *Filename);
 VECTORVM_API void    FreeVectorVMSerializeState      (FVectorVMSerializeState *SerializeState);
 
@@ -491,8 +607,8 @@ struct FVectorVMState {
 
 };
 
-VECTORVM_API uint32 SerializeVectorVMInputDataSets  (FVectorVMSerializeState *SerializeState, TArrayView<FDataSetMeta> DataSets, struct FVectorVMConstData *ConstData, int NumConstData); //only use when not calling InitVectorVMState()
-VECTORVM_API uint32 SerializeVectorVMOutputDataSets (FVectorVMSerializeState *SerializeState, TArrayView<FDataSetMeta> DataSets, struct FVectorVMConstData *ConstData, int NumConstData);
+VECTORVM_API uint32  SerializeVectorVMInputDataSets (FVectorVMSerializeState *SerializeState, TArrayView<FDataSetMeta>, const uint8 * const *ConstantTableData, const int *ConstantTableSizes, int32 ConstantTableCount); //only use when not calling InitVectorVMState()
+VECTORVM_API uint32  SerializeVectorVMOutputDataSets(FVectorVMSerializeState *SerializeState, TArrayView<FDataSetMeta>, const uint8 * const *ConstantTableData, const int *ConstantTableSizes, int32 ConstantTableCount);
 VECTORVM_API void   SerializeVectorVMWriteToFile    (FVectorVMSerializeState *SerializeState, uint8 WhichStateWritten, const wchar_t *Filename);
 VECTORVM_API void   FreeVectorVMSerializeState      (FVectorVMSerializeState *SerializeState);
 
