VectorVM Fluid Sim fix #jira UE-162851 #rb Rob.Krajcarski #preflight none

diff --git a/Engine/Source/Runtime/VectorVM/Private/VectorVMExperimental_Optimizer.inl b/Engine/Source/Runtime/VectorVM/Private/VectorVMExperimental_Optimizer.inl
--- a/Engine/Source/Runtime/VectorVM/Private/VectorVMExperimental_Optimizer.inl
+++ b/Engine/Source/Runtime/VectorVM/Private/VectorVMExperimental_Optimizer.inl
@@ -1209,12 +1209,21 @@ VECTORVM_API uint32 OptimizeVectorVMScript(const uint8 *InBytecode, int InByteco
 					++NumAcquireIndexInstructions;
 					int AcquireIndexInstructionIdx = i;
 					int NumInstructions = GetInstructionDependencyChain(OptContext, AcquireIndexInstructionIdx, RegToCheckStack, InstructionIdxStack);
-					//bubble up the dependent instructions at quickly as possible (@NOTE: if these are already grouped together it's more efficient to move more than 1 instruction at a time, but who cares? I doubt the performace will ever matter)
+					//bubble up the dependent instructions at quickly as possible
 					for (int j = 0; j < NumInstructions; ++j)
 					{
 						if (InstructionIdxStack[j] > LowestInstructionIdxForAcquireIdx)
 						{
 							FVectorVMOptimizeInstruction TempIns = OptContext->Intermediate.Instructions[InstructionIdxStack[j]];
+							if (TempIns.OpCode == EVectorVMOp::external_func_call)
+							{
+								for (int k = LowestInstructionIdxForAcquireIdx; k < InstructionIdxStack[j]; ++k)
+								{
+									if (OptContext->Intermediate.Instructions[k].OpCode == EVectorVMOp::external_func_call) {
+										LowestInstructionIdxForAcquireIdx = k + 1;
+									}
+								}
+							}
 							FMemory::Memmove(OptContext->Intermediate.Instructions + LowestInstructionIdxForAcquireIdx + 1, OptContext->Intermediate.Instructions + LowestInstructionIdxForAcquireIdx, sizeof(FVectorVMOptimizeInstruction) * (InstructionIdxStack[j] - LowestInstructionIdxForAcquireIdx));
 							OptContext->Intermediate.Instructions[LowestInstructionIdxForAcquireIdx] = TempIns;
 						}
@@ -1252,23 +1261,22 @@ VECTORVM_API uint32 OptimizeVectorVMScript(const uint8 *InBytecode, int InByteco
 						FVectorVMOptimizeInstruction *Ins = OptContext->Intermediate.Instructions + i;
 						FVectorVMOptimizeInsRegUsage RegUsage;
 						int NumRegisters = GetRegistersUsedForInstruction(OptContext, Ins, &RegUsage);
-						{ //if (((Ins->OpCat == EVectorVMOpCategory::Input && Ins->Input.FirstInsInsertIdx != -1) || Ins->OpCat != EVectorVMOpCategory::Input) && Ins->OpCat != EVectorVMOpCategory::Output) {
-							for (int j = 0; j < RegUsage.NumOutputRegisters; ++j)
+						
+						for (int j = 0; j < RegUsage.NumOutputRegisters; ++j)
+						{
+							if (OptContext->Intermediate.SSARegisterUsageBuffer[RegUsage.RegIndices[RegUsage.NumInputRegisters + j]] == IdxReg)
 							{
-								if (OptContext->Intermediate.SSARegisterUsageBuffer[RegUsage.RegIndices[RegUsage.NumInputRegisters + j]] == IdxReg)
-								{
-									FoundAcquireIndex = true;
-									OutputInsertionIdx = i + 1;
-								}
+								FoundAcquireIndex = true;
+								OutputInsertionIdx = i + 1;
 							}
-							if (FoundAcquireIndex)
+						}
+						if (FoundAcquireIndex)
+						{
+							for (int j = 0; j < NumRegisters; ++j)
 							{
-								for (int j = 0; j < NumRegisters; ++j)
+								if (OptContext->Intermediate.SSARegisterUsageBuffer[RegUsage.RegIndices[j]] == SrcReg)
 								{
-									if (OptContext->Intermediate.SSARegisterUsageBuffer[RegUsage.RegIndices[j]] == SrcReg)
-									{
-										OutputInsertionIdx = i + 1;
-									}
+									OutputInsertionIdx = i + 1;
 								}
 							}
 						}
