Fix for acos instruction with EVVM

diff --git a/Engine/Source/Runtime/VectorVM/Private/VectorVMExperimental.inl b/Engine/Source/Runtime/VectorVM/Private/VectorVMExperimental.inl
--- a/Engine/Source/Runtime/VectorVM/Private/VectorVMExperimental.inl
+++ b/Engine/Source/Runtime/VectorVM/Private/VectorVMExperimental.inl
@@ -376,7 +376,6 @@ VVMSet_m256iConst(  AlmostTwoBits                        , 0x3fffffff);
 #define VVM_vecBoolToFloat(v)              VectorSelect(v, VVM_m128Const(One), VectorZeroFloat());
 #define VVM_vecIntToBool(v)                VectorIntCompareGT(v, VectorSetZero())
 #define VVM_vecBoolToInt(v)                VectorIntSelect(v, VectorIntSet1(1), VectorSetZero())
-#define VVM_vecACosFast(v)                 VectorATan2(VectorSqrt(VectorMultiply(VectorSubtract(VVM_m128Const(One), v), VectorAdd(VVM_m128Const(One), v))), v)
 //safe instructions -- handle divide by zero "gracefully" by returning 0
 #define VVM_safeIns_div(v0, v1)            VectorSelect(VectorCompareGT(VectorAbs(v1), VVM_m128Const(Epsilon)), VectorDivide(v0, v1)                  , VectorZeroFloat())
 #define VVM_safeIns_rcp(v)                 VectorSelect(VectorCompareGT(VectorAbs(v) , VVM_m128Const(Epsilon)), VectorVMAccuracy::Reciprocal(v)       , VectorZeroFloat())
@@ -386,6 +385,7 @@ VVMSet_m256iConst(  AlmostTwoBits                        , 0x3fffffff);
 #define VVM_safe_rsq(v)                    VectorSelect(VectorCompareGT(v            , VVM_m128Const(Epsilon)), VectorReciprocalSqrt(v)               , VectorZeroFloat())
 #define VVM_random(v)                      VectorMultiply(VectorSubtract(VectorRegister4f(VectorCastIntToFloat(VectorIntOr(VectorShiftRightImmLogical(VVMXorwowStep(BatchState), 9), VectorIntSet1(0x3F800000)))), VVM_m128Const(One)), v)
 #define VVM_randomi(v)                     VectorFloatToInt(VectorMultiply(VectorSubtract(VectorRegister4f(VectorCastIntToFloat(VectorIntOr(VectorShiftRightImmLogical(VVMXorwowStep(BatchState), 9), VectorIntSet1(0x3F800000)))), VVM_m128Const(One)), *(VectorRegister4f *)&v))
+#define VVM_vecACosFast(v)                 VectorATan2(VVM_safe_sqrt(VectorMultiply(VectorSubtract(VVM_m128Const(One), v), VectorAdd(VVM_m128Const(One), v))), v)
 
 //new merged instructions
 #define VVM_cmplt_select(v0, v1, v2, v3)    VectorSelect(VectorCompareLT(v0, v1), v2, v3)
