Mirror safety net for pow EVVM instruction that was used on legacy VM (ensure base is greater than epsilon, otherwise return 0

diff --git a/Engine/Source/Runtime/VectorVM/Private/VectorVMExperimental.inl b/Engine/Source/Runtime/VectorVM/Private/VectorVMExperimental.inl
--- a/Engine/Source/Runtime/VectorVM/Private/VectorVMExperimental.inl
+++ b/Engine/Source/Runtime/VectorVM/Private/VectorVMExperimental.inl
@@ -381,7 +381,7 @@ VVMSet_m256iConst(  AlmostTwoBits                        , 0x3fffffff);
 #define VVM_safeIns_rcp(v)                 VectorSelect(VectorCompareGT(VectorAbs(v) , VVM_m128Const(Epsilon)), VectorVMAccuracy::Reciprocal(v)       , VectorZeroFloat())
 #define VVM_safe_sqrt(v)                   VectorSelect(VectorCompareGT(v            , VVM_m128Const(Epsilon)), VectorVMAccuracy::Sqrt(v)             , VectorZeroFloat())
 #define VVM_safe_log(v)                    VectorSelect(VectorCompareGT(v            , VectorZeroFloat())     , VectorLog(v)		                  , VectorZeroFloat())
-#define VVM_safe_pow(v0, v1)               VectorSelect(VectorCompareGT(v1           , VVM_m128Const(Epsilon)), VectorPow(v0, v1)                     , VectorZeroFloat())
+#define VVM_safe_pow(v0, v1)               VectorSelect(VectorCompareGT(v0           , VVM_m128Const(Epsilon)), VectorPow(v0, v1)                     , VectorZeroFloat())
 #define VVM_safe_rsq(v)                    VectorSelect(VectorCompareGT(v            , VVM_m128Const(Epsilon)), VectorReciprocalSqrt(v)               , VectorZeroFloat())
 #define VVM_random(v)                      VectorMultiply(VectorSubtract(VectorRegister4f(VectorCastIntToFloat(VectorIntOr(VectorShiftRightImmLogical(VVMXorwowStep(BatchState), 9), VectorIntSet1(0x3F800000)))), VVM_m128Const(One)), v)
 #define VVM_randomi(v)                     VectorFloatToInt(VectorMultiply(VectorSubtract(VectorRegister4f(VectorCastIntToFloat(VectorIntOr(VectorShiftRightImmLogical(VVMXorwowStep(BatchState), 9), VectorIntSet1(0x3F800000)))), VVM_m128Const(One)), *(VectorRegister4f *)&v))
