[NO NON-PII COMMIT MESSAGE]

diff --git a/Engine/Source/Runtime/VectorVM/Public/VectorVM.h b/Engine/Source/Runtime/VectorVM/Public/VectorVM.h
--- a/Engine/Source/Runtime/VectorVM/Public/VectorVM.h
+++ b/Engine/Source/Runtime/VectorVM/Public/VectorVM.h
@@ -216,7 +216,7 @@ namespace VectorVM
 #if VECTORVM_SUPPORTS_EXPERIMENTAL && VECTORVM_SUPPORTS_LEGACY
 			if (Context.UsingExperimentalVM)
 			{
-				InputPtr = (T*)Context.Experimental.GetNextRegister(&AdvanceOffset, &InputOffset) + Context.Experimental.PerInstanceFnInstanceIdx;
+				InputPtr = (T*)Context.Experimental.GetNextRegister(&AdvanceOffset, &InputOffset) + Context.Experimental.PerInstanceFnInstanceIdx * AdvanceOffset;
 				bIsRegister = !!AdvanceOffset;
 			}
 			else
@@ -232,7 +232,7 @@ namespace VectorVM
 				InputPtr += Context.Legacy.GetExternalFunctionInstanceOffset() * AdvanceOffset;
 			}
 #elif VECTORVM_SUPPORTS_EXPERIMENTAL
-			InputPtr = (T*)Context.GetNextRegister(&AdvanceOffset, &InputOffset) + Context.PerInstanceFnInstanceIdx;
+			InputPtr = (T*)Context.GetNextRegister(&AdvanceOffset, &InputOffset) + Context.PerInstanceFnInstanceIdx * AdvanceOffset;
 #elif VECTORVM_SUPPORTS_LEGACY
 			InputOffset = Context.DecodeU16();
 
@@ -290,7 +290,7 @@ namespace VectorVM
 		{
 			if (Context.UsingExperimentalVM)
 			{
-				Register = (T*)Context.Experimental.GetNextRegister(&AdvanceOffset, &RegisterIndex) + Context.Experimental.PerInstanceFnInstanceIdx;
+				Register = (T*)Context.Experimental.GetNextRegister(&AdvanceOffset, &RegisterIndex) + Context.Experimental.PerInstanceFnInstanceIdx * AdvanceOffset;
 			}
 			else
 			{
@@ -314,7 +314,7 @@ namespace VectorVM
 		}
 #elif VECTORVM_SUPPORTS_EXPERIMENTAL
 		FORCEINLINE FExternalFuncRegisterHandler(FVectorVMExternalFunctionContext& Context) {
-			Register = (T*)Context.GetNextRegister(&AdvanceOffset, &RegisterIndex) + Context.PerInstanceFnInstanceIdx;
+			Register = (T*)Context.GetNextRegister(&AdvanceOffset, &RegisterIndex) + Context.PerInstanceFnInstanceIdx * AdvanceOffset;
 		}
 #elif VECTORVM_SUPPORTS_LEGACY
 		FORCEINLINE FExternalFuncRegisterHandler(FVectorVMExternalFunctionContext& Context)
