Separated EVVM transcendental instructions from UnrealMath* intrinsics #preflight none #rb none

diff --git a/Engine/Source/Runtime/VectorVM/Private/VectorVMExecCore.inl b/Engine/Source/Runtime/VectorVM/Private/VectorVMExecCore.inl
--- a/Engine/Source/Runtime/VectorVM/Private/VectorVMExecCore.inl
+++ b/Engine/Source/Runtime/VectorVM/Private/VectorVMExecCore.inl
@@ -59,17 +59,17 @@ for (;;)
 		VVM_OP_CASE(sqrt)                             VVM_execVecIns1f(VVM_safe_sqrt);        VVM_OP_NEXT;
 		VVM_OP_CASE(neg)                              VVM_execVecIns1f(VectorNegate);         VVM_OP_NEXT;
 		VVM_OP_CASE(abs)                              VVM_execVecIns1f(VectorAbs);            VVM_OP_NEXT;
-		VVM_OP_CASE(exp)                              VVM_execVecIns1f(VectorExp);            VVM_OP_NEXT;
-		VVM_OP_CASE(exp2)                             VVM_execVecIns1f(VectorExp2);           VVM_OP_NEXT;
+		VVM_OP_CASE(exp)                              VVM_execVecIns1f(VVM_exp);              VVM_OP_NEXT;
+		VVM_OP_CASE(exp2)                             VVM_execVecIns1f(VVM_exp2);             VVM_OP_NEXT;
 		VVM_OP_CASE(log)                              VVM_execVecIns1f(VVM_safe_log);         VVM_OP_NEXT;
-		VVM_OP_CASE(log2)                             VVM_execVecIns1f(VectorLog2);           VVM_OP_NEXT;
-		VVM_OP_CASE(sin)                              VVM_execVecIns1f(VectorSin);            VVM_OP_NEXT;
-		VVM_OP_CASE(cos)                              VVM_execVecIns1f(VectorCos);            VVM_OP_NEXT;
-		VVM_OP_CASE(tan)                              VVM_execVecIns1f(VectorTan);            VVM_OP_NEXT;
-		VVM_OP_CASE(asin)                             VVM_execVecIns1f(VectorASin);           VVM_OP_NEXT;
+		VVM_OP_CASE(log2)                             VVM_execVecIns1f(VVM_log2);             VVM_OP_NEXT;
+		VVM_OP_CASE(sin)                              VVM_execVecIns1f(VVM_sin);              VVM_OP_NEXT;
+		VVM_OP_CASE(cos)                              VVM_execVecIns1f(VVM_cos);              VVM_OP_NEXT;
+		VVM_OP_CASE(tan)                              VVM_execVecIns1f(VVM_tan);              VVM_OP_NEXT;
+		VVM_OP_CASE(asin)                             VVM_execVecIns1f(VVM_vecASinFast);      VVM_OP_NEXT;
 		VVM_OP_CASE(acos)                             VVM_execVecIns1f(VVM_vecACosFast);      VVM_OP_NEXT;
-		VVM_OP_CASE(atan)                             VVM_execVecIns1f(VectorATan);           VVM_OP_NEXT;
-		VVM_OP_CASE(atan2)                            VVM_execVecIns2f(VectorATan2);          VVM_OP_NEXT;
+		VVM_OP_CASE(atan)                             VVM_execVecIns1f(VVM_atan);             VVM_OP_NEXT;
+		VVM_OP_CASE(atan2)                            VVM_execVecIns2f(VVM_atan2);            VVM_OP_NEXT;
 		VVM_OP_CASE(ceil)                             VVM_execVecIns1f(VectorCeil);           VVM_OP_NEXT;
 		VVM_OP_CASE(floor)                            VVM_execVecIns1f(VectorFloor);          VVM_OP_NEXT;
 		VVM_OP_CASE(fmod)                             VVM_execVecIns2f(VectorMod);            VVM_OP_NEXT;
diff --git a/Engine/Source/Runtime/VectorVM/Private/VectorVMExperimental.inl b/Engine/Source/Runtime/VectorVM/Private/VectorVMExperimental.inl
--- a/Engine/Source/Runtime/VectorVM/Private/VectorVMExperimental.inl
+++ b/Engine/Source/Runtime/VectorVM/Private/VectorVMExperimental.inl
@@ -1,5 +1,5 @@
 // Copyright Epic Games, Inc. All Rights Reserved.
-
+PRAGMA_DISABLE_OPTIMIZATION
 /*
 external functions to improve:
 UNiagaraDataInterfaceSkeletalMesh stuff (4 wide randoms)
@@ -344,8 +344,11 @@ VVMSet_m128Const(   NegativeOne     , -1.f);
 VVMSet_m128Const(   OneHalf         , 0.5f);
 VVMSet_m128Const(   Epsilon         , 1.e-8f);
 VVMSet_m128Const(   HalfPi          , 3.14159265359f * 0.5f);
+VVMSet_m128Const(   QuarterPi       , 3.14159265359f * 0.25f);
 VVMSet_m128Const(   FastSinA        , 7.5894663844f);
 VVMSet_m128Const(   FastSinB        , 1.6338434578f);
+VVMSet_m128Const(   Log2            , 0.6931471806f);
+VVMSet_m128Const(   OneOverLog2     , 1.4426950409f);
 VVMSet_m128iConst(  FMask           , 0xFFFFFFFF);
 VVMSet_m128iConst4( ZeroOneTwoThree , 0, 1, 2, 3);
 VVMSet_m128iConst4( ZeroTwoFourSix  , 0, 2, 4, 6);
@@ -371,6 +374,49 @@ VVMSet_m256iConst(  RegOffsetMask                        , 0x7FFF);
 VVMSet_m256Const(   RegOneOverTwoPi                      , 1.f / 2.f / 3.14159265359f);
 VVMSet_m256iConst(  AlmostTwoBits                        , 0x3fffffff);
 
+#if UE_PLATFORM_MATH_SSE
+#	if UE_PLATFORM_MATH_USE_SVML
+#		define VVM_tan(v)			            _mm_tan_ps(v)
+#		define VVM_atan(v)			            _mm_atan_ps(v)
+#		define VVM_atan2(v0, v1)	            _mm_atan2_ps(v0, v1)
+#	    define VVM_log2(v)						_mm_log2_ps(v)
+#	    define VVM_log(v)                       _mm_log_ps(v)
+#	    define VVM_exp(v)                       _mm_exp_ps(v)
+#		define VVM_exp2(v)                      _mm_exp2_ps(v)
+#	    define VVM_sin(v)                       _mm_sin_ps(v)
+#	    define VVM_cos(v)                       _mm_cos_ps(v)
+#	    define VVM_cot(v)                       _mm_cot_ps(v)
+#	    define VVM_sincos(v, out_sin, out_cos) *out_sin = _mm_sincos_ps(out_cos, v)
+#	else //UE_PLATFORM_MATH_USE_SVML
+#		define VVM_tan(v)			           SSE::tan_ps(v)
+#		define VVM_atan(v)			           SSE::atan_ps(v)
+#		define VVM_atan2(v0, v1)	           SSE::atan2_ps(v0, v1)
+#		define VVM_log2(v)                     _mm_mul_ps(SSE::log_ps(v), VVM_m128Const(OneOverLog2))
+#		define VVM_log(v)                      SSE::log_ps(v)
+#		define VVM_exp(v)                      SSE::exp_ps(v)
+#		define VVM_exp2(v)                     SSE::exp_ps(_mm_mul_ps(v, VVM_m128Const(Log2)))
+#		define VVM_sin(v)                      SSE::sin_ps(v)
+#		define VVM_cos(v)                      SSE::cos_ps(v)
+#		define VVM_cot(v)                      SSE::cot_ps(v)
+#		define VVM_sincos(v, out_sin, out_cos) SSE::sincos_ps(v, out_sin, out_cos)
+#		endif //UE_PLATFORM_MATH_SSE
+#elif UE_PLATFORM_MATH_NEON
+#	define VVM_tan(v)			               VectorTan(v)
+#	define VVM_atan(v)			               VectorATan(v)
+#	define VVM_atan2(v0, v1)	               VectorATan2(v)
+#	define VVM_log2(v)                         VectorLog2(v)
+#	define VVM_log(v)                          VectorLog(v)
+#	define VVM_exp(v)                          VectorExp(v)
+#	define VVM_exp(v)                          VectorExp2(v)
+#	define VVM_sin(v)                          VectorSin(v)
+#	define VVM_cos(v)                          VectorCos(v)
+#	define VVM_cot(v)                          VectorCot(v)
+#	define VVM_sincos(v, out_sin, out_cos)     VectorSinCos(out_sin, out_cos, *(v))
+#else
+#	error "VectorVM must use either SSE or NEON"
+#endif //UE_PLATFORM_MATH_NEON
+
+
 #define VVM_vecStep(a, b)                  VectorStep(VectorSubtract(a, b))
 #define VVM_vecFloatToBool(v)              VectorCompareGT(v, VectorZeroFloat())
 #define VVM_vecBoolToFloat(v)              VectorSelect(v, VVM_m128Const(One), VectorZeroFloat());
@@ -385,7 +431,8 @@ VVMSet_m256iConst(  AlmostTwoBits                        , 0x3fffffff);
 #define VVM_safe_rsq(v)                    VectorSelect(VectorCompareGT(v            , VVM_m128Const(Epsilon)), VectorReciprocalSqrt(v)               , VectorZeroFloat())
 #define VVM_random(v)                      VectorMultiply(VectorSubtract(VectorRegister4f(VectorCastIntToFloat(VectorIntOr(VectorShiftRightImmLogical(VVMXorwowStep(BatchState), 9), VectorIntSet1(0x3F800000)))), VVM_m128Const(One)), v)
 #define VVM_randomi(v)                     VectorFloatToInt(VectorMultiply(VectorSubtract(VectorRegister4f(VectorCastIntToFloat(VectorIntOr(VectorShiftRightImmLogical(VVMXorwowStep(BatchState), 9), VectorIntSet1(0x3F800000)))), VVM_m128Const(One)), *(VectorRegister4f *)&v))
-#define VVM_vecACosFast(v)                 VectorATan2(VVM_safe_sqrt(VectorMultiply(VectorSubtract(VVM_m128Const(One), v), VectorAdd(VVM_m128Const(One), v))), v)
+#define VVM_vecACosFast(v)                 VVM_atan2(VVM_safe_sqrt(VectorMultiply(VectorSubtract(VVM_m128Const(One), v), VectorAdd(VVM_m128Const(One), v))), v)
+#define VVM_vecASinFast(v)                 VectorSubtract(VVM_m128Const(QuarterPi), VVM_vecACosFast(v))
 
 //new merged instructions
 #define VVM_cmplt_select(v0, v1, v2, v3)    VectorSelect(VectorCompareLT(v0, v1), v2, v3)
@@ -640,7 +687,7 @@ union VVMXXMYMMUnion
 #define VVMIntRShift_AVX(v0, v1)                _mm256_srlv_epi32(v0, v1)
 #define VVMf2i_AVX(v)                           _mm256_cvttps_epi32(_mm256_castsi256_ps(v))
 #define VVMi2f_AVX(v)                           _mm256_cvtepi32_ps(_mm256_castps_si256(v))
-#define VectorSinCos_AVX(v, OutSin, OutCos)     sincos256_ps(v, OutSin, OutCos)
+#define VVM_sincos_AVX(v, OutSin, OutCos)       sincos256_ps(v, OutSin, OutCos)
 
 static inline __m256 VectorTan_AVX(const __m256 &v)
 {
@@ -881,6 +928,59 @@ FORCEINLINE void VVMMemSet32(void *dst, uint32 val, size_t num_vals)
 #include "VectorVMExperimental_Memory.inl"
 #pragma warning(pop)
 
+static bool vvmSkip3(EVectorVMOp op) {
+	return op == EVectorVMOp::cmplt_logic_and ||
+			op == EVectorVMOp::cmple_logic_and ||
+			op == EVectorVMOp::cmpgt_logic_and ||
+			op == EVectorVMOp::cmpge_logic_and ||
+			op == EVectorVMOp::cmpeq_logic_and ||
+			op == EVectorVMOp::cmpne_logic_and ||
+
+			op == EVectorVMOp::cmplti_logic_and ||
+			op == EVectorVMOp::cmplei_logic_and ||
+			op == EVectorVMOp::cmpgti_logic_and ||
+			op == EVectorVMOp::cmpgei_logic_and ||
+			op == EVectorVMOp::cmpeqi_logic_and ||
+			op == EVectorVMOp::cmpnei_logic_and ||
+
+			op == EVectorVMOp::cmplt_logic_or ||
+			op == EVectorVMOp::cmple_logic_or ||
+			op == EVectorVMOp::cmpgt_logic_or ||
+			op == EVectorVMOp::cmpge_logic_or ||
+			op == EVectorVMOp::cmpeq_logic_or ||
+			op == EVectorVMOp::cmpne_logic_or ||
+
+			op == EVectorVMOp::cmplti_logic_or ||
+			op == EVectorVMOp::cmplei_logic_or ||
+			op == EVectorVMOp::cmpgti_logic_or ||
+			op == EVectorVMOp::cmpgei_logic_or ||
+			op == EVectorVMOp::cmpeqi_logic_or ||
+			op == EVectorVMOp::cmpnei_logic_or;
+}
+
+#define VVM_FLOAT_CHECK1 { float *v[2]; uint8 inc[1]; v[0] = (float *)p0; v[1] = (float *)p1; inc[0] = BatchState->RegIncTable[((uint16 *)InsPtr)[0]]; vvmFloatCheck(v, inc, 1, BatchState->NumInstances, 0); }
+#define VVM_FLOAT_CHECK2 { float *v[3]; uint8 inc[2]; v[0] = (float *)p0; v[1] = (float *)p1; v[2] = (float *)p2; inc[0] = BatchState->RegIncTable[((uint16 *)InsPtr)[0]]; inc[1] = BatchState->RegIncTable[((uint16 *)InsPtr)[1]]; vvmFloatCheck(v, inc, 2, BatchState->NumInstances, 0); }
+#define VVM_FLOAT_CHECK3 { float *v[4]; uint8 inc[3]; v[0] = (float *)p0; v[1] = (float *)p1; v[2] = (float *)p2; v[3] = (float *)p3; inc[0] = BatchState->RegIncTable[((uint16 *)InsPtr)[0]]; inc[1] = BatchState->RegIncTable[((uint16 *)InsPtr)[1]]; inc[2] = BatchState->RegIncTable[((uint16 *)InsPtr)[2]]; if ((EVectorVMOp)InsPtr[-1] != EVectorVMOp::select) {vvmFloatCheck(v, inc, vvmSkip3((EVectorVMOp)InsPtr[-1]) ? 2 : 3, BatchState->NumInstances, 0); } }
+#define VVM_FLOAT_CHECK4 { float *v[5]; uint8 inc[4]; v[0] = (float *)p0; v[1] = (float *)p1; v[2] = (float *)p2; v[3] = (float *)p3; v[4] = (float *)p4; inc[0] = BatchState->RegIncTable[((uint16 *)InsPtr)[0]]; inc[1] = BatchState->RegIncTable[((uint16 *)InsPtr)[1]]; inc[2] = BatchState->RegIncTable[((uint16 *)InsPtr)[2]]; inc[3] = BatchState->RegIncTable[((uint16 *)InsPtr)[3]]; if ((EVectorVMOp)InsPtr[-1] != EVectorVMOp::cmplt_select && (EVectorVMOp)InsPtr[-1] != EVectorVMOp::cmple_select && (EVectorVMOp)InsPtr[-1] != EVectorVMOp::cmpeq_select && (EVectorVMOp)InsPtr[-1] != EVectorVMOp::cmplti_select && (EVectorVMOp)InsPtr[-1] != EVectorVMOp::cmplei_select && (EVectorVMOp)InsPtr[-1] != EVectorVMOp::cmpeqi_select && (EVectorVMOp)InsPtr[-1] != EVectorVMOp::select_mul && (EVectorVMOp)InsPtr[-1] != EVectorVMOp::select_add) { vvmFloatCheck(v, inc, 4, BatchState->NumInstances, 0); } }
+#define VVM_FLOAT_CHECK5 { float *v[6]; uint8 inc[5]; v[0] = (float *)p0; v[1] = (float *)p1; v[2] = (float *)p2; v[3] = (float *)p3; v[4] = (float *)p4; v[5] = (float *)p5; inc[0] = BatchState->RegIncTable[((uint16 *)InsPtr)[0]]; inc[1] = BatchState->RegIncTable[((uint16 *)InsPtr)[1]]; inc[2] = BatchState->RegIncTable[((uint16 *)InsPtr)[2]]; inc[3] = BatchState->RegIncTable[((uint16 *)InsPtr)[3]]; inc[4] = BatchState->RegIncTable[((uint16 *)InsPtr)[4]]; vvmFloatCheck(v, inc, 5, BatchState->NumInstances, 0); }
+
+static void vvmFloatCheck(float **vals, uint8 *inc, int num_vals, int num_instances, int start) {
+	//input regs
+	for (int i = start; i < num_vals; ++i) {
+		float *v = vals[i];
+		for (int j = 0; j < num_instances; ++j) {
+ 			int c = fpclassify(*v);
+			check(c == FP_NORMAL || c == FP_ZERO || c == FP_SUBNORMAL);
+			v += inc[i] >> 4;
+		}
+	}
+	//check output
+	for (int j = 0; j < num_instances; ++j) {
+		int c = fpclassify(vals[num_vals][j]);
+		check(c == FP_NORMAL || c == FP_ZERO || c == FP_SUBNORMAL);
+	}
+}
+
 #if VECTORVM_SUPPORTS_AVX
 static void SetRegIncTableFor8WideFrom4Wide(FVectorVMExecContext *ExecCtx, FVectorVMBatchState *BatchState)
 {
@@ -1585,3 +1685,4 @@ VECTORVM_API void ExecVectorVMState(FVectorVMState *VVMState, FVectorVMSerialize
 
 #endif //NIAGARA_EXP_VM
 
+PRAGMA_ENABLE_OPTIMIZATION
