Build health fix for previous submission. EVVM fix. #preflight 636921f24d3c1d9d925ce533 #rb none

diff --git a/Engine/Source/Runtime/VectorVM/Private/VectorVMExperimental.inl b/Engine/Source/Runtime/VectorVM/Private/VectorVMExperimental.inl
--- a/Engine/Source/Runtime/VectorVM/Private/VectorVMExperimental.inl
+++ b/Engine/Source/Runtime/VectorVM/Private/VectorVMExperimental.inl
@@ -1,5 +1,5 @@
 // Copyright Epic Games, Inc. All Rights Reserved.
-PRAGMA_DISABLE_OPTIMIZATION
+
 /*
 external functions to improve:
 UNiagaraDataInterfaceSkeletalMesh stuff (4 wide randoms)
@@ -374,47 +374,47 @@ VVMSet_m256iConst(  RegOffsetMask                        , 0x7FFF);
 VVMSet_m256Const(   RegOneOverTwoPi                      , 1.f / 2.f / 3.14159265359f);
 VVMSet_m256iConst(  AlmostTwoBits                        , 0x3fffffff);
 
-#if UE_PLATFORM_MATH_SSE
-#	if UE_PLATFORM_MATH_USE_SVML
-#		define VVM_tan(v)			            _mm_tan_ps(v)
-#		define VVM_atan(v)			            _mm_atan_ps(v)
-#		define VVM_atan2(v0, v1)	            _mm_atan2_ps(v0, v1)
-#	    define VVM_log2(v)						_mm_log2_ps(v)
-#	    define VVM_log(v)                       _mm_log_ps(v)
-#	    define VVM_exp(v)                       _mm_exp_ps(v)
-#		define VVM_exp2(v)                      _mm_exp2_ps(v)
-#	    define VVM_sin(v)                       _mm_sin_ps(v)
-#	    define VVM_cos(v)                       _mm_cos_ps(v)
-#	    define VVM_cot(v)                       _mm_cot_ps(v)
-#	    define VVM_sincos(v, out_sin, out_cos) *out_sin = _mm_sincos_ps(out_cos, v)
-#	else //UE_PLATFORM_MATH_USE_SVML
-#		define VVM_tan(v)			           SSE::tan_ps(v)
-#		define VVM_atan(v)			           SSE::atan_ps(v)
-#		define VVM_atan2(v0, v1)	           SSE::atan2_ps(v0, v1)
-#		define VVM_log2(v)                     _mm_mul_ps(SSE::log_ps(v), VVM_m128Const(OneOverLog2))
-#		define VVM_log(v)                      SSE::log_ps(v)
-#		define VVM_exp(v)                      SSE::exp_ps(v)
-#		define VVM_exp2(v)                     SSE::exp_ps(_mm_mul_ps(v, VVM_m128Const(Log2)))
-#		define VVM_sin(v)                      SSE::sin_ps(v)
-#		define VVM_cos(v)                      SSE::cos_ps(v)
-#		define VVM_cot(v)                      SSE::cot_ps(v)
-#		define VVM_sincos(v, out_sin, out_cos) SSE::sincos_ps(v, out_sin, out_cos)
-#		endif //UE_PLATFORM_MATH_SSE
-#elif UE_PLATFORM_MATH_NEON
+//#ifdef UE_PLATFORM_MATH_SSE
+//#	if UE_PLATFORM_MATH_USE_SVML
+//#		define VVM_tan(v)			            _mm_tan_ps(v)
+//#		define VVM_atan(v)			            _mm_atan_ps(v)
+//#		define VVM_atan2(v0, v1)	            _mm_atan2_ps(v0, v1)
+//#	    define VVM_log2(v)						_mm_log2_ps(v)
+//#	    define VVM_log(v)                       _mm_log_ps(v)
+//#	    define VVM_exp(v)                       _mm_exp_ps(v)
+//#		define VVM_exp2(v)                      _mm_exp2_ps(v)
+//#	    define VVM_sin(v)                       _mm_sin_ps(v)
+//#	    define VVM_cos(v)                       _mm_cos_ps(v)
+//#	    define VVM_cot(v)                       _mm_cot_ps(v)
+//#	    define VVM_sincos(v, out_sin, out_cos) *out_sin = _mm_sincos_ps(out_cos, v)
+//#	else //UE_PLATFORM_MATH_USE_SVML
+//#		define VVM_tan(v)			           SSE::tan_ps(v)
+//#		define VVM_atan(v)			           SSE::atan_ps(v)
+//#		define VVM_atan2(v0, v1)	           SSE::atan2_ps(v0, v1)
+//#		define VVM_log2(v)                     _mm_mul_ps(SSE::log_ps(v), VVM_m128Const(OneOverLog2))
+//#		define VVM_log(v)                      SSE::log_ps(v)
+//#		define VVM_exp(v)                      SSE::exp_ps(v)
+//#		define VVM_exp2(v)                     SSE::exp_ps(_mm_mul_ps(v, VVM_m128Const(Log2)))
+//#		define VVM_sin(v)                      SSE::sin_ps(v)
+//#		define VVM_cos(v)                      SSE::cos_ps(v)
+//#		define VVM_cot(v)                      SSE::cot_ps(v)
+//#		define VVM_sincos(v, out_sin, out_cos) SSE::sincos_ps(v, out_sin, out_cos)
+//#		endif //UE_PLATFORM_MATH_SSE
+//#elif UE_PLATFORM_MATH_NEON
 #	define VVM_tan(v)			               VectorTan(v)
 #	define VVM_atan(v)			               VectorATan(v)
-#	define VVM_atan2(v0, v1)	               VectorATan2(v)
+#	define VVM_atan2(v0, v1)	               VectorATan2(v0, v1)
 #	define VVM_log2(v)                         VectorLog2(v)
 #	define VVM_log(v)                          VectorLog(v)
 #	define VVM_exp(v)                          VectorExp(v)
-#	define VVM_exp(v)                          VectorExp2(v)
+#	define VVM_exp2(v)                         VectorExp2(v)
 #	define VVM_sin(v)                          VectorSin(v)
 #	define VVM_cos(v)                          VectorCos(v)
 #	define VVM_cot(v)                          VectorCot(v)
-#	define VVM_sincos(v, out_sin, out_cos)     VectorSinCos(out_sin, out_cos, *(v))
-#else
-#	error "VectorVM must use either SSE or NEON"
-#endif //UE_PLATFORM_MATH_NEON
+#	define VVM_sincos(v, out_sin, out_cos)     VectorSinCos(out_sin, out_cos, &(v))
+//#else
+//#	error "VectorVM must use either SSE or NEON"
+//#endif //UE_PLATFORM_MATH_NEON
 
 
 #define VVM_vecStep(a, b)                  VectorStep(VectorSubtract(a, b))
@@ -958,14 +958,26 @@ static bool vvmSkip3(EVectorVMOp op) {
 			op == EVectorVMOp::cmpnei_logic_or;
 }
 
-#define VVM_FLOAT_CHECK1 { float *v[2]; uint8 inc[1]; v[0] = (float *)p0; v[1] = (float *)p1; inc[0] = BatchState->RegIncTable[((uint16 *)InsPtr)[0]]; vvmFloatCheck(v, inc, 1, BatchState->NumInstances, 0); }
-#define VVM_FLOAT_CHECK2 { float *v[3]; uint8 inc[2]; v[0] = (float *)p0; v[1] = (float *)p1; v[2] = (float *)p2; inc[0] = BatchState->RegIncTable[((uint16 *)InsPtr)[0]]; inc[1] = BatchState->RegIncTable[((uint16 *)InsPtr)[1]]; vvmFloatCheck(v, inc, 2, BatchState->NumInstances, 0); }
-#define VVM_FLOAT_CHECK3 { float *v[4]; uint8 inc[3]; v[0] = (float *)p0; v[1] = (float *)p1; v[2] = (float *)p2; v[3] = (float *)p3; inc[0] = BatchState->RegIncTable[((uint16 *)InsPtr)[0]]; inc[1] = BatchState->RegIncTable[((uint16 *)InsPtr)[1]]; inc[2] = BatchState->RegIncTable[((uint16 *)InsPtr)[2]]; if ((EVectorVMOp)InsPtr[-1] != EVectorVMOp::select) {vvmFloatCheck(v, inc, vvmSkip3((EVectorVMOp)InsPtr[-1]) ? 2 : 3, BatchState->NumInstances, 0); } }
-#define VVM_FLOAT_CHECK4 { float *v[5]; uint8 inc[4]; v[0] = (float *)p0; v[1] = (float *)p1; v[2] = (float *)p2; v[3] = (float *)p3; v[4] = (float *)p4; inc[0] = BatchState->RegIncTable[((uint16 *)InsPtr)[0]]; inc[1] = BatchState->RegIncTable[((uint16 *)InsPtr)[1]]; inc[2] = BatchState->RegIncTable[((uint16 *)InsPtr)[2]]; inc[3] = BatchState->RegIncTable[((uint16 *)InsPtr)[3]]; if ((EVectorVMOp)InsPtr[-1] != EVectorVMOp::cmplt_select && (EVectorVMOp)InsPtr[-1] != EVectorVMOp::cmple_select && (EVectorVMOp)InsPtr[-1] != EVectorVMOp::cmpeq_select && (EVectorVMOp)InsPtr[-1] != EVectorVMOp::cmplti_select && (EVectorVMOp)InsPtr[-1] != EVectorVMOp::cmplei_select && (EVectorVMOp)InsPtr[-1] != EVectorVMOp::cmpeqi_select && (EVectorVMOp)InsPtr[-1] != EVectorVMOp::select_mul && (EVectorVMOp)InsPtr[-1] != EVectorVMOp::select_add) { vvmFloatCheck(v, inc, 4, BatchState->NumInstances, 0); } }
-#define VVM_FLOAT_CHECK5 { float *v[6]; uint8 inc[5]; v[0] = (float *)p0; v[1] = (float *)p1; v[2] = (float *)p2; v[3] = (float *)p3; v[4] = (float *)p4; v[5] = (float *)p5; inc[0] = BatchState->RegIncTable[((uint16 *)InsPtr)[0]]; inc[1] = BatchState->RegIncTable[((uint16 *)InsPtr)[1]]; inc[2] = BatchState->RegIncTable[((uint16 *)InsPtr)[2]]; inc[3] = BatchState->RegIncTable[((uint16 *)InsPtr)[3]]; inc[4] = BatchState->RegIncTable[((uint16 *)InsPtr)[4]]; vvmFloatCheck(v, inc, 5, BatchState->NumInstances, 0); }
+#if 0
+#define VVM_FLOAT_CHECK1 const uint16 *float_check_start_ins_ptr = (uint16 *)(InsPtr + 1); { float *v[2]; uint8 inc[1]; v[0] = (float *)p0; v[1] = (float *)p1; inc[0] = BatchState->RegIncTable[((uint16 *)InsPtr)[0]]; vvmFloatCheck(v, inc, 1, BatchState->NumInstances, 0); }
+#define VVM_FLOAT_CHECK2 const uint16 *float_check_start_ins_ptr = (uint16 *)(InsPtr + 1); { float *v[3]; uint8 inc[2]; v[0] = (float *)p0; v[1] = (float *)p1; v[2] = (float *)p2; inc[0] = BatchState->RegIncTable[((uint16 *)InsPtr)[0]]; inc[1] = BatchState->RegIncTable[((uint16 *)InsPtr)[1]]; vvmFloatCheck(v, inc, 2, BatchState->NumInstances, 0); }
+#define VVM_FLOAT_CHECK3 const uint16 *float_check_start_ins_ptr = (uint16 *)(InsPtr + 1); { float *v[4]; uint8 inc[3]; v[0] = (float *)p0; v[1] = (float *)p1; v[2] = (float *)p2; v[3] = (float *)p3; inc[0] = BatchState->RegIncTable[((uint16 *)InsPtr)[0]]; inc[1] = BatchState->RegIncTable[((uint16 *)InsPtr)[1]]; inc[2] = BatchState->RegIncTable[((uint16 *)InsPtr)[2]]; if ((EVectorVMOp)InsPtr[-1] != EVectorVMOp::select) {vvmFloatCheck(v, inc, vvmSkip3((EVectorVMOp)InsPtr[-1]) ? 2 : 3, BatchState->NumInstances, 0); } }
+#define VVM_FLOAT_CHECK4 const uint16 *float_check_start_ins_ptr = (uint16 *)(InsPtr + 1); { float *v[5]; uint8 inc[4]; v[0] = (float *)p0; v[1] = (float *)p1; v[2] = (float *)p2; v[3] = (float *)p3; v[4] = (float *)p4; inc[0] = BatchState->RegIncTable[((uint16 *)InsPtr)[0]]; inc[1] = BatchState->RegIncTable[((uint16 *)InsPtr)[1]]; inc[2] = BatchState->RegIncTable[((uint16 *)InsPtr)[2]]; inc[3] = BatchState->RegIncTable[((uint16 *)InsPtr)[3]]; if ((EVectorVMOp)InsPtr[-1] != EVectorVMOp::cmplt_select && (EVectorVMOp)InsPtr[-1] != EVectorVMOp::cmple_select && (EVectorVMOp)InsPtr[-1] != EVectorVMOp::cmpeq_select && (EVectorVMOp)InsPtr[-1] != EVectorVMOp::cmplti_select && (EVectorVMOp)InsPtr[-1] != EVectorVMOp::cmplei_select && (EVectorVMOp)InsPtr[-1] != EVectorVMOp::cmpeqi_select && (EVectorVMOp)InsPtr[-1] != EVectorVMOp::select_mul && (EVectorVMOp)InsPtr[-1] != EVectorVMOp::select_add) { vvmFloatCheck(v, inc, 4, BatchState->NumInstances, 0); } }
+#define VVM_FLOAT_CHECK5 const uint16 *float_check_start_ins_ptr = (uint16 *)(InsPtr + 1); { float *v[6]; uint8 inc[5]; v[0] = (float *)p0; v[1] = (float *)p1; v[2] = (float *)p2; v[3] = (float *)p3; v[4] = (float *)p4; v[5] = (float *)p5; inc[0] = BatchState->RegIncTable[((uint16 *)InsPtr)[0]]; inc[1] = BatchState->RegIncTable[((uint16 *)InsPtr)[1]]; inc[2] = BatchState->RegIncTable[((uint16 *)InsPtr)[2]]; inc[3] = BatchState->RegIncTable[((uint16 *)InsPtr)[3]]; inc[4] = BatchState->RegIncTable[((uint16 *)InsPtr)[4]]; vvmFloatCheck(v, inc, 5, BatchState->NumInstances, 0); }
+
+#define VVM_FLOAT_CHECK_RES(ResultRegIdx) { float *v = (float *)BatchState->RegPtrTable[float_check_start_ins_ptr[ResultRegIdx]]; uint8 inc = BatchState->RegIncTable[float_check_start_ins_ptr[ResultRegIdx]]; vvmFloatCheck(&v, &inc, 1, BatchState->NumInstances, 0); }
+#else
+#define VVM_FLOAT_CHECK1
+#define VVM_FLOAT_CHECK2
+#define VVM_FLOAT_CHECK3
+#define VVM_FLOAT_CHECK4
+#define VVM_FLOAT_CHECK5
+
+#define VVM_FLOAT_CHECK_RES(...)
+#endif
+
 
 static void vvmFloatCheck(float **vals, uint8 *inc, int num_vals, int num_instances, int start) {
-	//input regs
 	for (int i = start; i < num_vals; ++i) {
 		float *v = vals[i];
 		for (int j = 0; j < num_instances; ++j) {
@@ -974,11 +986,6 @@ static void vvmFloatCheck(float **vals, uint8 *inc, int num_vals, int num_instan
 			v += inc[i] >> 4;
 		}
 	}
-	//check output
-	for (int j = 0; j < num_instances; ++j) {
-		int c = fpclassify(vals[num_vals][j]);
-		check(c == FP_NORMAL || c == FP_ZERO || c == FP_SUBNORMAL);
-	}
 }
 
 #if VECTORVM_SUPPORTS_AVX
@@ -1685,4 +1692,3 @@ VECTORVM_API void ExecVectorVMState(FVectorVMState *VVMState, FVectorVMSerialize
 
 #endif //NIAGARA_EXP_VM
 
-PRAGMA_ENABLE_OPTIMIZATION
