Guard against integer overflow with integer divide within VectorVM

diff --git a/Engine/Source/Runtime/VectorVM/Private/VectorVM.cpp b/Engine/Source/Runtime/VectorVM/Private/VectorVM.cpp
--- a/Engine/Source/Runtime/VectorVM/Private/VectorVM.cpp
+++ b/Engine/Source/Runtime/VectorVM/Private/VectorVM.cpp
@@ -11,6 +11,7 @@
 #include "GenericPlatform/GenericPlatformFile.h"
 #include "Async/ParallelFor.h"
 #include "Math/UnrealPlatformMathSSE.h"
+#include <limits>
 
 IMPLEMENT_MODULE(FDefaultModuleImpl, VectorVM);
 
@@ -2316,15 +2317,32 @@ struct FVectorIntKernelDivide : TBinaryVectorIntKernel<FVectorIntKernelDivide>
 		VectorIntStore(Src1, TmpB);
 
 		// No intrinsics exist for integer divide. Since div by zero causes crashes, we must be safe against that.
-
 		int32 TmpDst[4];
-		TmpDst[0] = TmpB[0] != 0 ? (TmpA[0] / TmpB[0]) : 0;
-		TmpDst[1] = TmpB[1] != 0 ? (TmpA[1] / TmpB[1]) : 0;
-		TmpDst[2] = TmpB[2] != 0 ? (TmpA[2] / TmpB[2]) : 0;
-		TmpDst[3] = TmpB[3] != 0 ? (TmpA[3] / TmpB[3]) : 0;
+		TmpDst[0] = SafeIntDivide(TmpA[0], TmpB[0]);
+		TmpDst[1] = SafeIntDivide(TmpA[1], TmpB[1]);
+		TmpDst[2] = SafeIntDivide(TmpA[2], TmpB[2]);
+		TmpDst[3] = SafeIntDivide(TmpA[3], TmpB[3]);
 
 		*Dst = MakeVectorRegisterInt(TmpDst[0], TmpDst[1], TmpDst[2], TmpDst[3]);
 	}
+
+private:
+	VM_FORCEINLINE static int32 SafeIntDivide(int32 Numerator, int32 Denominator)
+	{
+		static constexpr int32 MinIntValue = std::numeric_limits<int32>::min();
+		static constexpr int32 MaxIntValue = std::numeric_limits<int32>::max();
+
+		if (Denominator == 0)
+		{
+			return 0;
+		}
+		else if ((Denominator == -1) && (Numerator == MinIntValue))
+		{
+			return MaxIntValue;
+		}
+
+		return Numerator / Denominator;
+	}
 };
 
 
