Fix for experimental VM not supporting FExternalFuncConstHandler

diff --git a/Engine/Source/Runtime/VectorVM/Public/VectorVM.h b/Engine/Source/Runtime/VectorVM/Public/VectorVM.h
--- a/Engine/Source/Runtime/VectorVM/Public/VectorVM.h
+++ b/Engine/Source/Runtime/VectorVM/Public/VectorVM.h
@@ -290,7 +290,8 @@ namespace VectorVM
 		{
 			if (Context.UsingExperimentalVM)
 			{
-				Register = (T*)Context.Experimental.GetNextRegister(&AdvanceOffset, &RegisterIndex) + Context.Experimental.PerInstanceFnInstanceIdx * AdvanceOffset;
+				Register = (T*)Context.Experimental.GetNextRegister(&AdvanceOffset, &RegisterIndex);
+				Register += Context.Experimental.PerInstanceFnInstanceIdx * AdvanceOffset;
 			}
 			else
 			{
@@ -313,8 +314,10 @@ namespace VectorVM
 			}
 		}
 #elif VECTORVM_SUPPORTS_EXPERIMENTAL
-		FORCEINLINE FExternalFuncRegisterHandler(FVectorVMExternalFunctionContext& Context) {
-			Register = (T*)Context.GetNextRegister(&AdvanceOffset, &RegisterIndex) + Context.PerInstanceFnInstanceIdx * AdvanceOffset;
+		FORCEINLINE FExternalFuncRegisterHandler(FVectorVMExternalFunctionContext& Context)
+		{
+			Register = (T*)Context.GetNextRegister(&AdvanceOffset, &RegisterIndex);
+			Register += Context.PerInstanceFnInstanceIdx * AdvanceOffset;
 		}
 #elif VECTORVM_SUPPORTS_LEGACY
 		FORCEINLINE FExternalFuncRegisterHandler(FVectorVMExternalFunctionContext& Context)
@@ -360,7 +363,6 @@ namespace VectorVM
 	template<typename T>
 	struct FExternalFuncConstHandler
 	{
-		uint16 ConstantIndex;
 		T Constant;
 
 #if VECTORVM_SUPPORTS_EXPERIMENTAL && VECTORVM_SUPPORTS_LEGACY
@@ -368,23 +370,32 @@ namespace VectorVM
 		{
 			if (Context.UsingExperimentalVM)
 			{
-				check(false);
+				int32 RegisterIndex;
+				int32 AdvanceOffset;
+				const T* Register = (const T*)Context.Experimental.GetNextRegister(&AdvanceOffset, &RegisterIndex);
+				Register += Context.Experimental.PerInstanceFnInstanceIdx * AdvanceOffset;
+
+				Constant = *Register;
 			}
 			else
 			{
-				ConstantIndex = (Context.Legacy.DecodeU16() & VVM_EXT_FUNC_INPUT_LOC_MASK);
+				const uint16 ConstantIndex = (Context.Legacy.DecodeU16() & VVM_EXT_FUNC_INPUT_LOC_MASK);
 				Constant = (*Context.Legacy.GetConstant<T>(ConstantIndex));
 			}
 		}
 #elif VECTORVM_SUPPORTS_EXPERIMENTAL
 		FExternalFuncConstHandler(FVectorVMExternalFunctionContext& Context)
 		{
-			check(false);
+			int32 RegisterIndex;
+			int32 AdvanceOffset;
+			const T* Register = (const T*)Context.GetNextRegister(&AdvanceOffset, &RegisterIndex);
+			Register += Context.PerInstanceFnInstanceIdx * AdvanceOffset;
+
+			Constant = *Register;
 		}
 #elif VECTORVM_SUPPORTS_LEGACY
 		FExternalFuncConstHandler(FVectorVMExternalFunctionContext& Context)
-			: ConstantIndex(Context.DecodeU16() & VVM_EXT_FUNC_INPUT_LOC_MASK)
-			, Constant(*Context.GetConstant<T>(ConstantIndex))
+			: Constant(*Context.GetConstant<T>(Context.DecodeU16()& VVM_EXT_FUNC_INPUT_LOC_MASK))
 		{}
 #endif
 		FORCEINLINE const T& Get() { return Constant; }
