Small optimization to Legacy VM (contributed by licensee [REDACTED_PERSON]

diff --git a/Engine/Source/Runtime/VectorVM/Private/VectorVM.cpp b/Engine/Source/Runtime/VectorVM/Private/VectorVM.cpp
--- a/Engine/Source/Runtime/VectorVM/Private/VectorVM.cpp
+++ b/Engine/Source/Runtime/VectorVM/Private/VectorVM.cpp
@@ -3340,10 +3340,10 @@ struct FBatchedWriteIndexedOutput
 		}
 	}
 
-	template<typename SourceType, typename TargetType, int32 TypeOffset>
+	template<typename SourceType, typename TargetType, int32 TypeOffset, bool bCheckIfEmpty = true>
 	static void CopyRegisterToOutput(FVectorVMContext& Context)
 	{
-		if (SkipIfEmpty(Context))
+		if (bCheckIfEmpty && SkipIfEmpty(Context))
 		{
 			return;
 		}
@@ -3375,7 +3375,7 @@ struct FBatchedWriteIndexedOutput
 		}
 		else if (Context.ValidInstanceUniform)
 		{
-			CopyRegisterToOutput<SourceType, TargetType, TypeOffset>(Context);
+			CopyRegisterToOutput<SourceType, TargetType, TypeOffset, false>(Context);
 			return;
 		}
 
