- Add VVM support for bit casting between floats and ints

diff --git a/Engine/Source/Developer/ShaderFormatVectorVM/Private/ir_vm_gen_bytecode_visitor.cpp b/Engine/Source/Developer/ShaderFormatVectorVM/Private/ir_vm_gen_bytecode_visitor.cpp
--- a/Engine/Source/Developer/ShaderFormatVectorVM/Private/ir_vm_gen_bytecode_visitor.cpp
+++ b/Engine/Source/Developer/ShaderFormatVectorVM/Private/ir_vm_gen_bytecode_visitor.cpp
@@ -192,8 +192,10 @@ void BuildExpressionMap()
 		* bit pattern casting operations
 		*/
 		// 		Info = &VMExpressionMap.Add(ir_unop_fasu);//TODO?
-		// 		Info = &VMExpressionMap.Add(ir_unop_fasi);
-		// 		Info = &VMExpressionMap.Add(ir_unop_iasf);
+		Info = &VMExpressionMap.Add(ir_unop_fasi);
+		Info->Add(FVMExpresssionInfo(EVectorVMOp::fasi, glsl_type::int_type, glsl_type::float_type));
+		Info = &VMExpressionMap.Add(ir_unop_iasf);
+		Info->Add(FVMExpresssionInfo(EVectorVMOp::iasf, glsl_type::float_type, glsl_type::int_type));
 		// 		Info = &VMExpressionMap.Add(ir_unop_uasf);
 
 		/**
diff --git a/Engine/Source/Runtime/Core/Private/UObject/FortniteMainBranchObjectVersion.cpp b/Engine/Source/Runtime/Core/Private/UObject/FortniteMainBranchObjectVersion.cpp
--- a/Engine/Source/Runtime/Core/Private/UObject/FortniteMainBranchObjectVersion.cpp
+++ b/Engine/Source/Runtime/Core/Private/UObject/FortniteMainBranchObjectVersion.cpp
@@ -11,7 +11,7 @@ TMap<FGuid, FGuid> FFortniteMainBranchObjectVersion::GetSystemGuids()
 	SystemGuids.Add(DevGuids.MATERIALSHADERMAP_DERIVEDDATA_VER, FGuid("CE7776E63E9E43FBB8C103EE1AB40B5D"));
 	SystemGuids.Add(DevGuids.NANITE_DERIVEDDATA_VER, FGuid("D28C003F825F4C178F698938FB0E47F7"));
 	SystemGuids.Add(DevGuids.NIAGARASHADERMAP_DERIVEDDATA_VER, FGuid("A6FFCB143BBC4C83A72001468D23D5A4"));
-	SystemGuids.Add(DevGuids.Niagara_LatestScriptCompileVersion, FGuid("555EEFA140B7604B99DC2122BB27D3F4"));
+	SystemGuids.Add(DevGuids.Niagara_LatestScriptCompileVersion, FGuid("F8277B215893A046A04BADAEE853CB81"));
 	SystemGuids.Add(DevGuids.SkeletalMeshDerivedDataVersion, FGuid("518D7C51D31D9BE4A9F4B81542AC8119"));
 	SystemGuids.Add(DevGuids.STATICMESH_DERIVEDDATA_VER, FGuid("543843B103794E0E9BA4BE80FB602F79"));
 	return SystemGuids;
diff --git a/Engine/Source/Runtime/VectorVM/Private/VectorVM.cpp b/Engine/Source/Runtime/VectorVM/Private/VectorVM.cpp
--- a/Engine/Source/Runtime/VectorVM/Private/VectorVM.cpp
+++ b/Engine/Source/Runtime/VectorVM/Private/VectorVM.cpp
@@ -2636,6 +2636,25 @@ struct FVectorKernelBoolToInt : TUnaryKernel<FVectorKernelBoolToInt, FRegisterHa
 	}
 };
 
+//reinterpret bits
+//fasi,
+struct FVectorKernelFloatAsInt : TUnaryKernel<FVectorKernelFloatAsInt, FRegisterHandler<VectorRegister4Int>, FConstantHandler<VectorRegister4Float>, FRegisterHandler<VectorRegister4Float>, VECTOR_WIDTH_FLOATS>
+{
+	static void VM_FORCEINLINE DoKernel(FVectorVMContext& Context, VectorRegister4Int* Dst, VectorRegister4Float Src0)
+	{
+		*Dst = VectorCastFloatToInt(Src0);
+	}
+};
+
+//iasf,
+struct FVectorKernelIntAsFloat : TUnaryKernel<FVectorKernelIntAsFloat, FRegisterHandler<VectorRegister4Float>, FConstantHandler<VectorRegister4Int>, FRegisterHandler<VectorRegister4Int>, VECTOR_WIDTH_FLOATS>
+{
+	static void VM_FORCEINLINE DoKernel(FVectorVMContext& Context, VectorRegister4Float* Dst, VectorRegister4Int Src0)
+	{
+		*Dst = VectorCastIntToFloat(Src0);
+	}
+};
+
 void VectorVM::Exec(FVectorVMExecArgs& Args, FVectorVMSerializeState *SerializeState)
 {
 	//TRACE_CPUPROFILER_EVENT_SCOPE("VMExec");
@@ -2796,6 +2815,8 @@ void VectorVM::Exec(FVectorVMExecArgs& Args, FVectorVMSerializeState *SerializeS
 						case EVectorVMOp::b2f: FVectorKernelBoolToFloat::Exec(Context); break;
 						case EVectorVMOp::i2b: FVectorKernelIntToBool::Exec(Context); break;
 						case EVectorVMOp::b2i: FVectorKernelBoolToInt::Exec(Context); break;
+						case EVectorVMOp::fasi: FVectorKernelFloatAsInt::Exec(Context); break;
+						case EVectorVMOp::iasf: FVectorKernelIntAsFloat::Exec(Context); break;
 
 						case EVectorVMOp::outputdata_half:	FScalarKernelWriteOutputIndexed<float, FFloat16, 2>::Exec(Context);	break;
 						case EVectorVMOp::inputdata_half: FVectorKernelReadInput<FFloat16, 2>::Exec(Context); break;
@@ -3783,6 +3804,8 @@ void VectorVM::OptimizeByteCode(const uint8* ByteCode, TArray<uint8>& OptimizedC
 			case EVectorVMOp::b2f: FVectorKernelBoolToFloat::Optimize(Context); break;
 			case EVectorVMOp::i2b: FVectorKernelIntToBool::Optimize(Context); break;
 			case EVectorVMOp::b2i: FVectorKernelBoolToInt::Optimize(Context); break;
+			case EVectorVMOp::fasi: FVectorKernelFloatAsInt::Optimize(Context); break;
+			case EVectorVMOp::iasf: FVectorKernelIntAsFloat::Optimize(Context); break;
 
 			case EVectorVMOp::outputdata_half:	FScalarKernelWriteOutputIndexed<float, FFloat16, 2>::Optimize(Context);	break;
 			case EVectorVMOp::inputdata_half: FVectorKernelReadInput<FFloat16, 2>::Optimize(Context); break;
diff --git a/Engine/Source/Runtime/VectorVM/Private/VectorVMExecCore.inl b/Engine/Source/Runtime/VectorVM/Private/VectorVMExecCore.inl
--- a/Engine/Source/Runtime/VectorVM/Private/VectorVMExecCore.inl
+++ b/Engine/Source/Runtime/VectorVM/Private/VectorVMExecCore.inl
@@ -16,6 +16,7 @@ static const void* jmp_tbl[] = {
 #define VVM_OP_CASE(op)	case EVectorVMOp::op:
 #define VVM_OP_NEXT     break
 #endif
+#define VVM_NullOp(X) X
 
 { //set the pointers to deal with the correct chunk offset.
 	int DstOffset = ExecCtx->VVMState->NumTempRegisters + ExecCtx->VVMState->NumConstBuffers;
@@ -432,6 +433,9 @@ for (;;)
 				p1 += sizeof(FVecReg);
 			} while (p1 < end);
 		} VVM_OP_NEXT;
+		VVM_OP_CASE(fasi)                                VVM_execVecIns1i(VVM_NullOp);            VVM_OP_NEXT;
+		VVM_OP_CASE(iasf)                                VVM_execVecIns1f(VVM_NullOp);            VVM_OP_NEXT;
+
 		VVM_OP_CASE(exec_indexf)                         VVM_execVec_exec_indexf;                 VVM_OP_NEXT;
 		VVM_OP_CASE(exec_index_addi)                     VVM_execVec_exec_index_addi;             VVM_OP_NEXT;
 
diff --git a/Engine/Source/Runtime/VectorVM/Private/VectorVMExperimental_Optimizer.inl b/Engine/Source/Runtime/VectorVM/Private/VectorVMExperimental_Optimizer.inl
--- a/Engine/Source/Runtime/VectorVM/Private/VectorVMExperimental_Optimizer.inl
+++ b/Engine/Source/Runtime/VectorVM/Private/VectorVMExperimental_Optimizer.inl
@@ -759,6 +759,8 @@ VECTORVM_API uint32 OptimizeVectorVMScript(const uint8 *InBytecode, int InByteco
 			case EVectorVMOp::b2f:                                  VVMOptimizeVecIns1;             break;
 			case EVectorVMOp::i2b:                                  VVMOptimizeVecIns1;             break;
 			case EVectorVMOp::b2i:                                  VVMOptimizeVecIns1;             break;
+			case EVectorVMOp::fasi:                                 VVMOptimizeVecIns1;             break;
+			case EVectorVMOp::iasf:                                 VVMOptimizeVecIns1;             break;
 			case EVectorVMOp::inputdata_float:
 			case EVectorVMOp::inputdata_int32:
 			case EVectorVMOp::inputdata_half:
diff --git a/Engine/Source/Runtime/VectorVM/Public/VectorVMCommon.h b/Engine/Source/Runtime/VectorVM/Public/VectorVMCommon.h
--- a/Engine/Source/Runtime/VectorVM/Public/VectorVMCommon.h
+++ b/Engine/Source/Runtime/VectorVM/Public/VectorVMCommon.h
@@ -128,102 +128,104 @@ enum class EVectorVMOperandLocation : uint8
 	VVM_OP_XM( update_id                    , RWBuffer  , 3			, VVM_INS_PARAM_FFFFFF              ) /* 93 */                                                                                                         \
 	VVM_OP_XM( acquire_id                   , RWBuffer  , 3			, VVM_INS_PARAM_FFFFFF              ) /* 94 */                                                                                                         \
 	VVM_OP_XM( half_to_float                , Fused     , 2			, VVM_INS_PARAM_FFFFFF              ) /* 95 */                                                                                                         \
+	VVM_OP_XM( fasi                         , Op        , 2			, VVM_INS_PARAM_FFFFFF              ) /* 96 */                                                                                                         \
+	VVM_OP_XM( iasf                         , Op        , 2			, VVM_INS_PARAM_FFFFFF              ) /* 97 */                                                                                                         \
 	/* Fused ops -- combined ops that show up frequently together in Fornite.  There are three types: */                                                                                                                   \
 	/* 1. exec_index that get immediately fed into an add or i2f */                                                                                                                                                        \
 	/* 2. ops with identical inputs */                                                                                                                                                                                     \
 	/* 3. ops where the output of one chain to the next.. ie a mul that feeds directly into a sub */                                                                                                                       \
-	VVM_OP_XM( exec_indexf                  , Op        , 2         , VVM_INS_PARAM_FFIIIF              ) /* 96 */                                                                                                         \
-	VVM_OP_XM( exec_index_addi              , Op        , 2         , VVM_INS_PARAM_FFIIII              ) /* 97 */                                                                                                         \
+	VVM_OP_XM( exec_indexf                  , Op        , 2         , VVM_INS_PARAM_FFIIIF              ) /* 98 */                                                                                                         \
+	VVM_OP_XM( exec_index_addi              , Op        , 2         , VVM_INS_PARAM_FFIIII              ) /* 90 */                                                                                                         \
 	\
-	VVM_OP_XM( cmplt_select                 , Op        , 5         , VVM_INS_PARAM_FFFFFF              ) /* 98 */                                                                                                         \
-	VVM_OP_XM( cmple_select                 , Op        , 5         , VVM_INS_PARAM_FFFFFF              ) /* 99 */                                                                                                        \
-	VVM_OP_XM( cmpeq_select                 , Op        , 5         , VVM_INS_PARAM_FFFFFF              ) /* 100 */                                                                                                        \
-	VVM_OP_XM( cmplti_select                , Op        , 5         , VVM_INS_PARAM_FIIFFF              ) /* 101 */                                                                                                        \
-	VVM_OP_XM( cmplei_select                , Op        , 5         , VVM_INS_PARAM_FIIFFF              ) /* 102 */                                                                                                        \
-	VVM_OP_XM( cmpeqi_select                , Op        , 5         , VVM_INS_PARAM_FIIFFF              ) /* 103 */                                                                                                        \
+	VVM_OP_XM( cmplt_select                 , Op        , 5         , VVM_INS_PARAM_FFFFFF              ) /* 100 */                                                                                                         \
+	VVM_OP_XM( cmple_select                 , Op        , 5         , VVM_INS_PARAM_FFFFFF              ) /* 101 */                                                                                                        \
+	VVM_OP_XM( cmpeq_select                 , Op        , 5         , VVM_INS_PARAM_FFFFFF              ) /* 102 */                                                                                                        \
+	VVM_OP_XM( cmplti_select                , Op        , 5         , VVM_INS_PARAM_FIIFFF              ) /* 103 */                                                                                                        \
+	VVM_OP_XM( cmplei_select                , Op        , 5         , VVM_INS_PARAM_FIIFFF              ) /* 104 */                                                                                                        \
+	VVM_OP_XM( cmpeqi_select                , Op        , 5         , VVM_INS_PARAM_FIIFFF              ) /* 105 */                                                                                                        \
 	\
-	VVM_OP_XM( cmplt_logic_and              , Op        , 4         , VVM_INS_PARAM_FFFFII              ) /* 104 */                                                                                                        \
-	VVM_OP_XM( cmple_logic_and              , Op        , 4         , VVM_INS_PARAM_FFFFII              ) /* 105 */                                                                                                        \
-	VVM_OP_XM( cmpgt_logic_and              , Op        , 4         , VVM_INS_PARAM_FFFFII              ) /* 106 */                                                                                                        \
-	VVM_OP_XM( cmpge_logic_and              , Op        , 4         , VVM_INS_PARAM_FFFFII              ) /* 107 */                                                                                                        \
-	VVM_OP_XM( cmpeq_logic_and              , Op        , 4         , VVM_INS_PARAM_FFFFII              ) /* 108 */                                                                                                        \
-	VVM_OP_XM( cmpne_logic_and              , Op        , 4         , VVM_INS_PARAM_FFFFII              ) /* 109 */                                                                                                        \
-	VVM_OP_XM( cmplti_logic_and             , Op        , 4         , VVM_INS_PARAM_FIIFII              ) /* 110 */                                                                                                        \
-	VVM_OP_XM( cmplei_logic_and             , Op        , 4         , VVM_INS_PARAM_FIIFII              ) /* 111 */                                                                                                        \
-	VVM_OP_XM( cmpgti_logic_and             , Op        , 4         , VVM_INS_PARAM_FIIFII              ) /* 112 */                                                                                                        \
-	VVM_OP_XM( cmpgei_logic_and             , Op        , 4         , VVM_INS_PARAM_FIIFIF              ) /* 113 */                                                                                                        \
-	VVM_OP_XM( cmpeqi_logic_and             , Op        , 4         , VVM_INS_PARAM_FIIFIF              ) /* 114 */                                                                                                        \
-	VVM_OP_XM( cmpnei_logic_and             , Op        , 4         , VVM_INS_PARAM_FIIFIF              ) /* 115 */                                                                                                        \
+	VVM_OP_XM( cmplt_logic_and              , Op        , 4         , VVM_INS_PARAM_FFFFII              ) /* 106 */                                                                                                        \
+	VVM_OP_XM( cmple_logic_and              , Op        , 4         , VVM_INS_PARAM_FFFFII              ) /* 107 */                                                                                                        \
+	VVM_OP_XM( cmpgt_logic_and              , Op        , 4         , VVM_INS_PARAM_FFFFII              ) /* 108 */                                                                                                        \
+	VVM_OP_XM( cmpge_logic_and              , Op        , 4         , VVM_INS_PARAM_FFFFII              ) /* 109 */                                                                                                        \
+	VVM_OP_XM( cmpeq_logic_and              , Op        , 4         , VVM_INS_PARAM_FFFFII              ) /* 110 */                                                                                                        \
+	VVM_OP_XM( cmpne_logic_and              , Op        , 4         , VVM_INS_PARAM_FFFFII              ) /* 111 */                                                                                                        \
+	VVM_OP_XM( cmplti_logic_and             , Op        , 4         , VVM_INS_PARAM_FIIFII              ) /* 112 */                                                                                                        \
+	VVM_OP_XM( cmplei_logic_and             , Op        , 4         , VVM_INS_PARAM_FIIFII              ) /* 113 */                                                                                                        \
+	VVM_OP_XM( cmpgti_logic_and             , Op        , 4         , VVM_INS_PARAM_FIIFII              ) /* 114 */                                                                                                        \
+	VVM_OP_XM( cmpgei_logic_and             , Op        , 4         , VVM_INS_PARAM_FIIFIF              ) /* 115 */                                                                                                        \
+	VVM_OP_XM( cmpeqi_logic_and             , Op        , 4         , VVM_INS_PARAM_FIIFIF              ) /* 116 */                                                                                                        \
+	VVM_OP_XM( cmpnei_logic_and             , Op        , 4         , VVM_INS_PARAM_FIIFIF              ) /* 117 */                                                                                                        \
 	\
-	VVM_OP_XM( cmplt_logic_or               , Op        , 4         , VVM_INS_PARAM_FFFFFF              ) /* 116 */                                                                                                        \
-	VVM_OP_XM( cmple_logic_or               , Op        , 4         , VVM_INS_PARAM_FFFFFF              ) /* 117 */                                                                                                        \
-	VVM_OP_XM( cmpgt_logic_or               , Op        , 4         , VVM_INS_PARAM_FFFFFF              ) /* 118 */                                                                                                        \
-	VVM_OP_XM( cmpge_logic_or               , Op        , 4         , VVM_INS_PARAM_FFFFFF              ) /* 119 */                                                                                                        \
-	VVM_OP_XM( cmpeq_logic_or               , Op        , 4         , VVM_INS_PARAM_FFFFFF              ) /* 120 */                                                                                                        \
-	VVM_OP_XM( cmpne_logic_or               , Op        , 4         , VVM_INS_PARAM_FFFFFF              ) /* 121 */                                                                                                        \
-	VVM_OP_XM( cmplti_logic_or              , Op        , 4         , VVM_INS_PARAM_FIIFFF              ) /* 122 */                                                                                                        \
-	VVM_OP_XM( cmplei_logic_or              , Op        , 4         , VVM_INS_PARAM_FIIFFF              ) /* 123 */                                                                                                        \
-	VVM_OP_XM( cmpgti_logic_or              , Op        , 4         , VVM_INS_PARAM_FIIFFF              ) /* 124 */                                                                                                        \
-	VVM_OP_XM( cmpgei_logic_or              , Op        , 4         , VVM_INS_PARAM_FIIFFF              ) /* 125 */                                                                                                        \
-	VVM_OP_XM( cmpeqi_logic_or              , Op        , 4         , VVM_INS_PARAM_FIIFFF              ) /* 126 */                                                                                                        \
-	VVM_OP_XM( cmpnei_logic_or              , Op        , 4         , VVM_INS_PARAM_FIIFFF              ) /* 127 */                                                                                                        \
+	VVM_OP_XM( cmplt_logic_or               , Op        , 4         , VVM_INS_PARAM_FFFFFF              ) /* 118 */                                                                                                        \
+	VVM_OP_XM( cmple_logic_or               , Op        , 4         , VVM_INS_PARAM_FFFFFF              ) /* 119 */                                                                                                        \
+	VVM_OP_XM( cmpgt_logic_or               , Op        , 4         , VVM_INS_PARAM_FFFFFF              ) /* 120 */                                                                                                        \
+	VVM_OP_XM( cmpge_logic_or               , Op        , 4         , VVM_INS_PARAM_FFFFFF              ) /* 121 */                                                                                                        \
+	VVM_OP_XM( cmpeq_logic_or               , Op        , 4         , VVM_INS_PARAM_FFFFFF              ) /* 122 */                                                                                                        \
+	VVM_OP_XM( cmpne_logic_or               , Op        , 4         , VVM_INS_PARAM_FFFFFF              ) /* 123 */                                                                                                        \
+	VVM_OP_XM( cmplti_logic_or              , Op        , 4         , VVM_INS_PARAM_FIIFFF              ) /* 124 */                                                                                                        \
+	VVM_OP_XM( cmplei_logic_or              , Op        , 4         , VVM_INS_PARAM_FIIFFF              ) /* 125 */                                                                                                        \
+	VVM_OP_XM( cmpgti_logic_or              , Op        , 4         , VVM_INS_PARAM_FIIFFF              ) /* 126 */                                                                                                        \
+	VVM_OP_XM( cmpgei_logic_or              , Op        , 4         , VVM_INS_PARAM_FIIFFF              ) /* 127 */                                                                                                        \
+	VVM_OP_XM( cmpeqi_logic_or              , Op        , 4         , VVM_INS_PARAM_FIIFFF              ) /* 128 */                                                                                                        \
+	VVM_OP_XM( cmpnei_logic_or              , Op        , 4         , VVM_INS_PARAM_FIIFFF              ) /* 129 */                                                                                                        \
 	\
-	VVM_OP_XM( mad_add                      , Op        , 5         , VVM_INS_PARAM_FFFFFF              ) /* 128 */                                                                                                        \
-	VVM_OP_XM( mad_sub0                     , Op        , 5         , VVM_INS_PARAM_FFFFFF              ) /* 129 */                                                                                                        \
-	VVM_OP_XM( mad_sub1                     , Op        , 5         , VVM_INS_PARAM_FFFFFF              ) /* 130 */                                                                                                        \
-	VVM_OP_XM( mad_mul                      , Op        , 5         , VVM_INS_PARAM_FFFFFF              ) /* 131 */                                                                                                        \
-	VVM_OP_XM( mad_sqrt                     , Op        , 4         , VVM_INS_PARAM_FFFFFF              ) /* 132 */                                                                                                        \
-	VVM_OP_XM( mad_mad0                     , Op        , 6         , VVM_INS_PARAM_FFFFFF              ) /* 133 */                                                                                                        \
-	VVM_OP_XM( mad_mad1                     , Op        , 6         , VVM_INS_PARAM_FFFFFF              ) /* 134 */                                                                                                        \
+	VVM_OP_XM( mad_add                      , Op        , 5         , VVM_INS_PARAM_FFFFFF              ) /* 130 */                                                                                                        \
+	VVM_OP_XM( mad_sub0                     , Op        , 5         , VVM_INS_PARAM_FFFFFF              ) /* 131 */                                                                                                        \
+	VVM_OP_XM( mad_sub1                     , Op        , 5         , VVM_INS_PARAM_FFFFFF              ) /* 132 */                                                                                                        \
+	VVM_OP_XM( mad_mul                      , Op        , 5         , VVM_INS_PARAM_FFFFFF              ) /* 133 */                                                                                                        \
+	VVM_OP_XM( mad_sqrt                     , Op        , 4         , VVM_INS_PARAM_FFFFFF              ) /* 134 */                                                                                                        \
+	VVM_OP_XM( mad_mad0                     , Op        , 6         , VVM_INS_PARAM_FFFFFF              ) /* 135 */                                                                                                        \
+	VVM_OP_XM( mad_mad1                     , Op        , 6         , VVM_INS_PARAM_FFFFFF              ) /* 136 */                                                                                                        \
 	\
-	VVM_OP_XM( mul_mad0                     , Op        , 5         , VVM_INS_PARAM_FFFFFF              ) /* 135 */                                                                                                        \
-	VVM_OP_XM( mul_mad1                     , Op        , 5         , VVM_INS_PARAM_FFFFFF              ) /* 136 */                                                                                                        \
-	VVM_OP_XM( mul_add                      , Op        , 4         , VVM_INS_PARAM_FFFFFF              ) /* 137 */                                                                                                        \
-	VVM_OP_XM( mul_sub0                     , Op        , 4         , VVM_INS_PARAM_FFFFFF              ) /* 138 */                                                                                                        \
-	VVM_OP_XM( mul_sub1                     , Op        , 4         , VVM_INS_PARAM_FFFFFF              ) /* 139 */                                                                                                        \
-	VVM_OP_XM( mul_mul                      , Op        , 4         , VVM_INS_PARAM_FFFFFF              ) /* 140 */                                                                                                        \
-	VVM_OP_XM( mul_max                      , Op        , 4         , VVM_INS_PARAM_FFFFFF              ) /* 141 */                                                                                                        \
-	VVM_OP_XM( mul_2x                       , Op        , 4         , VVM_INS_PARAM_FFFFFF              ) /* 142 */                                                                                                        \
+	VVM_OP_XM( mul_mad0                     , Op        , 5         , VVM_INS_PARAM_FFFFFF              ) /* 137 */                                                                                                        \
+	VVM_OP_XM( mul_mad1                     , Op        , 5         , VVM_INS_PARAM_FFFFFF              ) /* 138 */                                                                                                        \
+	VVM_OP_XM( mul_add                      , Op        , 4         , VVM_INS_PARAM_FFFFFF              ) /* 139 */                                                                                                        \
+	VVM_OP_XM( mul_sub0                     , Op        , 4         , VVM_INS_PARAM_FFFFFF              ) /* 140 */                                                                                                        \
+	VVM_OP_XM( mul_sub1                     , Op        , 4         , VVM_INS_PARAM_FFFFFF              ) /* 141 */                                                                                                        \
+	VVM_OP_XM( mul_mul                      , Op        , 4         , VVM_INS_PARAM_FFFFFF              ) /* 142 */                                                                                                        \
+	VVM_OP_XM( mul_max                      , Op        , 4         , VVM_INS_PARAM_FFFFFF              ) /* 143 */                                                                                                        \
+	VVM_OP_XM( mul_2x                       , Op        , 4         , VVM_INS_PARAM_FFFFFF              ) /* 144 */                                                                                                        \
 	\
-	VVM_OP_XM( add_mad1                     , Op        , 5         , VVM_INS_PARAM_FFFFFF              ) /* 143 */                                                                                                        \
-	VVM_OP_XM( add_add                      , Op        , 4         , VVM_INS_PARAM_FFFFFF              ) /* 144 */                                                                                                        \
+	VVM_OP_XM( add_mad1                     , Op        , 5         , VVM_INS_PARAM_FFFFFF              ) /* 145 */                                                                                                        \
+	VVM_OP_XM( add_add                      , Op        , 4         , VVM_INS_PARAM_FFFFFF              ) /* 146 */                                                                                                        \
 	\
-	VVM_OP_XM( sub_cmplt1                   , Op        , 4         , VVM_INS_PARAM_FFFFFI              ) /* 145 */                                                                                                        \
-	VVM_OP_XM( sub_neg                      , Op        , 3         , VVM_INS_PARAM_FFFFFF              ) /* 146 */                                                                                                        \
-	VVM_OP_XM( sub_mul                      , Op        , 4         , VVM_INS_PARAM_FFFFFF              ) /* 147 */                                                                                                        \
+	VVM_OP_XM( sub_cmplt1                   , Op        , 4         , VVM_INS_PARAM_FFFFFI              ) /* 147 */                                                                                                        \
+	VVM_OP_XM( sub_neg                      , Op        , 3         , VVM_INS_PARAM_FFFFFF              ) /* 148 */                                                                                                        \
+	VVM_OP_XM( sub_mul                      , Op        , 4         , VVM_INS_PARAM_FFFFFF              ) /* 149 */                                                                                                        \
 	\
-	VVM_OP_XM( div_mad0                     , Op        , 5         , VVM_INS_PARAM_FFFFFF              ) /* 148 */                                                                                                        \
-	VVM_OP_XM( div_f2i                      , Op        , 3         , VVM_INS_PARAM_FFFFFI              ) /* 149 */                                                                                                        \
-	VVM_OP_XM( div_mul                      , Op        , 4         , VVM_INS_PARAM_FFFFFF              ) /* 150 */                                                                                                        \
+	VVM_OP_XM( div_mad0                     , Op        , 5         , VVM_INS_PARAM_FFFFFF              ) /* 150 */                                                                                                        \
+	VVM_OP_XM( div_f2i                      , Op        , 3         , VVM_INS_PARAM_FFFFFI              ) /* 151 */                                                                                                        \
+	VVM_OP_XM( div_mul                      , Op        , 4         , VVM_INS_PARAM_FFFFFF              ) /* 152 */                                                                                                        \
 	\
-	VVM_OP_XM( muli_addi                    , Op        , 4         , VVM_INS_PARAM_IIIIII              ) /* 151 */                                                                                                        \
-	VVM_OP_XM( addi_bit_rshift				, Op        , 4         , VVM_INS_PARAM_IIIIII              ) /* 152 */                                                                                                        \
-	VVM_OP_XM( addi_muli                    , Op        , 4         , VVM_INS_PARAM_IIIIII              ) /* 153 */                                                                                                        \
+	VVM_OP_XM( muli_addi                    , Op        , 4         , VVM_INS_PARAM_IIIIII              ) /* 153 */                                                                                                        \
+	VVM_OP_XM( addi_bit_rshift				, Op        , 4         , VVM_INS_PARAM_IIIIII              ) /* 154 */                                                                                                        \
+	VVM_OP_XM( addi_muli                    , Op        , 4         , VVM_INS_PARAM_IIIIII              ) /* 155 */                                                                                                        \
 	\
-	VVM_OP_XM( b2i_2x                       , Op        , 3         , VVM_INS_PARAM_IIIIII              ) /* 154 */                                                                                                        \
+	VVM_OP_XM( b2i_2x                       , Op        , 3         , VVM_INS_PARAM_IIIIII              ) /* 156 */                                                                                                        \
 	\
-	VVM_OP_XM( i2f_div0                     , Op        , 3         , VVM_INS_PARAM_FFFFFF              ) /* 155 */                                                                                                        \
-	VVM_OP_XM( i2f_div1                     , Op        , 3         , VVM_INS_PARAM_FFFFFF              ) /* 156 */                                                                                                        \
-	VVM_OP_XM( i2f_mul                      , Op        , 3         , VVM_INS_PARAM_FFFFFF              ) /* 157 */                                                                                                        \
-	VVM_OP_XM( i2f_mad0                     , Op        , 4         , VVM_INS_PARAM_FFFFFF              ) /* 158 */                                                                                                        \
-	VVM_OP_XM( i2f_mad1                     , Op        , 4         , VVM_INS_PARAM_FFFFFF              ) /* 159 */                                                                                                        \
+	VVM_OP_XM( i2f_div0                     , Op        , 3         , VVM_INS_PARAM_FFFFFF              ) /* 157 */                                                                                                        \
+	VVM_OP_XM( i2f_div1                     , Op        , 3         , VVM_INS_PARAM_FFFFFF              ) /* 158 */                                                                                                        \
+	VVM_OP_XM( i2f_mul                      , Op        , 3         , VVM_INS_PARAM_FFFFFF              ) /* 159 */                                                                                                        \
+	VVM_OP_XM( i2f_mad0                     , Op        , 4         , VVM_INS_PARAM_FFFFFF              ) /* 160 */                                                                                                        \
+	VVM_OP_XM( i2f_mad1                     , Op        , 4         , VVM_INS_PARAM_FFFFFF              ) /* 161 */                                                                                                        \
 	\
-	VVM_OP_XM( f2i_select1                  , Op        , 4         , VVM_INS_PARAM_IIFIII              ) /* 160 */                                                                                                        \
-	VVM_OP_XM( f2i_maxi                     , Op        , 3         , VVM_INS_PARAM_IIIFII              ) /* 161 */                                                                                                        \
-	VVM_OP_XM( f2i_addi                     , Op        , 3         , VVM_INS_PARAM_IIIFII              ) /* 162 */                                                                                                        \
+	VVM_OP_XM( f2i_select1                  , Op        , 4         , VVM_INS_PARAM_IIFIII              ) /* 162 */                                                                                                        \
+	VVM_OP_XM( f2i_maxi                     , Op        , 3         , VVM_INS_PARAM_IIIFII              ) /* 163 */                                                                                                        \
+	VVM_OP_XM( f2i_addi                     , Op        , 3         , VVM_INS_PARAM_IIIFII              ) /* 164 */                                                                                                        \
 	\
-	VVM_OP_XM( fmod_add                     , Op        , 3         , VVM_INS_PARAM_FFFFFF              ) /* 163 */                                                                                                        \
-	VVM_OP_XM( bit_and_i2f                  , Op        , 3         , VVM_INS_PARAM_FFFFFF              ) /* 164 */                                                                                                        \
-	VVM_OP_XM( bit_rshift_bit_and           , Op        , 3         , VVM_INS_PARAM_FFFFFF              ) /* 165 */                                                                                                        \
-	VVM_OP_XM( neg_cmplt                    , Op        , 3         , VVM_INS_PARAM_FFFFFF              ) /* 166 */                                                                                                        \
-	VVM_OP_XM( bit_or_muli                  , Op        , 3         , VVM_INS_PARAM_FFFFFF              ) /* 167 */                                                                                                        \
-	VVM_OP_XM( bit_lshift_bit_or            , Op        , 3         , VVM_INS_PARAM_FFFFFF              ) /* 168 */                                                                                                        \
-	VVM_OP_XM( random_add                   , Op        , 3         , VVM_INS_PARAM_FFFFFF              ) /* 169 */                                                                                                        \
-	VVM_OP_XM( random_2x                    , Op        , 3         , VVM_INS_PARAM_FFFFFF              ) /* 170 */                                                                                                        \
-	VVM_OP_XM( max_f2i                      , Op        , 3         , VVM_INS_PARAM_FFFFFF              ) /* 171 */                                                                                                        \
-	VVM_OP_XM( select_mul                   , Op        , 5         , VVM_INS_PARAM_FFFFFF              ) /* 172 */                                                                                                        \
-	VVM_OP_XM( select_add                   , Op        , 5         , VVM_INS_PARAM_FFFFFF              ) /* 173 */                                                                                                        \
-	VVM_OP_XM( sin_cos                      , Op        , 3         , VVM_INS_PARAM_FFFFFF              ) /* 174 */
+	VVM_OP_XM( fmod_add                     , Op        , 3         , VVM_INS_PARAM_FFFFFF              ) /* 165 */                                                                                                        \
+	VVM_OP_XM( bit_and_i2f                  , Op        , 3         , VVM_INS_PARAM_FFFFFF              ) /* 166 */                                                                                                        \
+	VVM_OP_XM( bit_rshift_bit_and           , Op        , 3         , VVM_INS_PARAM_FFFFFF              ) /* 167 */                                                                                                        \
+	VVM_OP_XM( neg_cmplt                    , Op        , 3         , VVM_INS_PARAM_FFFFFF              ) /* 168 */                                                                                                        \
+	VVM_OP_XM( bit_or_muli                  , Op        , 3         , VVM_INS_PARAM_FFFFFF              ) /* 169 */                                                                                                        \
+	VVM_OP_XM( bit_lshift_bit_or            , Op        , 3         , VVM_INS_PARAM_FFFFFF              ) /* 170 */                                                                                                        \
+	VVM_OP_XM( random_add                   , Op        , 3         , VVM_INS_PARAM_FFFFFF              ) /* 171 */                                                                                                        \
+	VVM_OP_XM( random_2x                    , Op        , 3         , VVM_INS_PARAM_FFFFFF              ) /* 172 */                                                                                                        \
+	VVM_OP_XM( max_f2i                      , Op        , 3         , VVM_INS_PARAM_FFFFFF              ) /* 173 */                                                                                                        \
+	VVM_OP_XM( select_mul                   , Op        , 5         , VVM_INS_PARAM_FFFFFF              ) /* 174 */                                                                                                        \
+	VVM_OP_XM( select_add                   , Op        , 5         , VVM_INS_PARAM_FFFFFF              ) /* 175 */                                                                                                        \
+	VVM_OP_XM( sin_cos                      , Op        , 3         , VVM_INS_PARAM_FFFFFF              ) /* 176 */
 
 #define VVM_OP_XM(n, ...) n,
 enum class EVectorVMOp : uint8 {
