Experimental Vector VM improvements: faster execution speed, faster optimization speed, uses less runtime memory, added debugging features

diff --git a/Engine/Source/Runtime/VectorVM/Private/VectorVMExperimentalSerialization.cpp b/Engine/Source/Runtime/VectorVM/Private/VectorVMExperimentalSerialization.cpp
new file mode 100644
--- /dev/null
+++ b/Engine/Source/Runtime/VectorVM/Private/VectorVMExperimentalSerialization.cpp
@@ -0,0 +1,1151 @@
+// Copyright Epic Games, Inc. All Rights Reserved.
+
+#include "HAL/FileManager.h"
+#include "HAL/PlatformFileManager.h"
+#include "VectorVMExperimental.h"
+
+
+void *VVMDefaultRealloc(void *Ptr, size_t NumBytes, const char *Filename, int LineNumber);
+void VVMDefaultFree(void *Ptr, const char *Filename, int LineNumber);
+
+
+#if VECTORVM_SUPPORTS_SERIALIZATION
+
+int VVMGetRegisterType(FVectorVMState *VVMState, uint16 RegIdx, uint16 *OutAbsReg = NULL)
+{
+	if (RegIdx < (int)VVMState->NumTempRegisters)
+	{
+		if (OutAbsReg)
+		{
+			*OutAbsReg = RegIdx;
+		}
+		return VVM_RT_TEMPREG;
+	}
+	else if (RegIdx < (int)(VVMState->NumTempRegisters + VVMState->NumConstBuffers))
+	{
+		if (OutAbsReg)
+		{
+			*OutAbsReg = RegIdx - VVMState->NumTempRegisters;
+		}
+		return VVM_RT_CONST;
+	}
+	else if (RegIdx < (int)(VVMState->NumTempRegisters + VVMState->NumConstBuffers + VVMState->NumInputBuffers * 2))
+	{
+		if (OutAbsReg)
+		{
+			*OutAbsReg = RegIdx - VVMState->NumTempRegisters - VVMState->NumConstBuffers;
+		}
+		return VVM_RT_INPUT;
+	}
+	else if (RegIdx < (int)(VVMState->NumTempRegisters + VVMState->NumConstBuffers + VVMState->NumInputBuffers * 2 + VVMState->NumOutputsRemapped))
+	{
+		if (OutAbsReg)
+		{
+			*OutAbsReg = RegIdx - VVMState->NumTempRegisters - VVMState->NumConstBuffers - VVMState->NumInputBuffers * 2;
+		}
+		return VVM_RT_OUTPUT;
+	}
+	return VVM_RT_INVALID;
+}
+
+
+//prototypes for stuff shared between the original and experimental VM
+static FVectorVMSerializeInstruction *VVMSerGetNextInstruction(FVectorVMSerializeState *SerializeState, int GlobalChunkIdx, int OpStart, int NumOps, uint64 Dt, uint64 DtDecode);
+
+void FreeVectorVMSerializeState(FVectorVMSerializeState *SerializeState)
+{
+	if (SerializeState->FreeFn == nullptr)
+	{
+		return;
+	}
+	for (uint32 i = 0; i < SerializeState->NumExternalData; ++i)
+	{
+		if (SerializeState->ExternalData[i].Name)
+		{
+			SerializeState->FreeFn(SerializeState->ExternalData[i].Name, __FILE__, __LINE__);
+		}
+	}
+	SerializeState->FreeFn(SerializeState->RegisterTableFlags, __FILE__, __LINE__);
+	SerializeState->FreeFn(SerializeState->ExternalData      , __FILE__, __LINE__);
+	SerializeState->FreeFn(SerializeState->Bytecode          , __FILE__, __LINE__);
+	SerializeState->FreeFn(SerializeState->Chunks            , __FILE__, __LINE__);
+
+	for (uint32 i = 0; i < SerializeState->NumDataSets; ++i)
+	{
+		if (SerializeState->DataSets[i].InputBuffers)
+		{
+			SerializeState->FreeFn(SerializeState->DataSets[i].InputBuffers, __FILE__, __LINE__);
+		}
+		if (SerializeState->DataSets[i].OutputBuffers)
+		{
+			SerializeState->FreeFn(SerializeState->DataSets[i].OutputBuffers, __FILE__, __LINE__);
+		}		
+		if (SerializeState->DataSets[i].InputIDTable)
+		{
+			SerializeState->FreeFn(SerializeState->DataSets[i].InputIDTable, __FILE__, __LINE__);
+		}
+		if (SerializeState->DataSets[i].InputFreeIDTable)
+		{
+			SerializeState->FreeFn(SerializeState->DataSets[i].InputFreeIDTable, __FILE__, __LINE__);
+		}
+		if (SerializeState->DataSets[i].InputSpawnedIDTable)
+		{
+			SerializeState->FreeFn(SerializeState->DataSets[i].InputSpawnedIDTable, __FILE__, __LINE__);
+		}
+		if (SerializeState->DataSets[i].OutputIDTable)
+		{
+			SerializeState->FreeFn(SerializeState->DataSets[i].OutputIDTable, __FILE__, __LINE__);
+		}
+		if (SerializeState->DataSets[i].OutputFreeIDTable)
+		{
+			SerializeState->FreeFn(SerializeState->DataSets[i].OutputFreeIDTable, __FILE__, __LINE__);
+		}
+		if (SerializeState->DataSets[i].OutputSpawnedIDTable)
+		{
+			SerializeState->FreeFn(SerializeState->DataSets[i].OutputSpawnedIDTable, __FILE__, __LINE__);
+		}
+	}
+	SerializeState->FreeFn(SerializeState->DataSets, __FILE__, __LINE__);
+	for (uint32 i = 0; i < SerializeState->NumInstructions; ++i)
+	{
+		SerializeState->FreeFn(SerializeState->Instructions[i].RegisterTable, __FILE__, __LINE__);
+		SerializeState->FreeFn(SerializeState->Instructions[i].RegisterFlags, __FILE__, __LINE__);
+	}
+	SerializeState->FreeFn(SerializeState->Instructions, __FILE__, __LINE__);
+	SerializeState->FreeFn(SerializeState->ConstTableSizesInBytes, __FILE__, __LINE__);
+	SerializeState->FreeFn(SerializeState->PreExecConstData      , __FILE__, __LINE__);
+	SerializeState->FreeFn(SerializeState->PostExecConstData     , __FILE__, __LINE__);
+	SerializeState->FreeFn(SerializeState->OutputRemapDataSetIdx , __FILE__, __LINE__);
+	SerializeState->FreeFn(SerializeState->OutputRemapDataType   , __FILE__, __LINE__);
+	SerializeState->FreeFn(SerializeState->OutputRemapDst        , __FILE__, __LINE__);
+}
+
+static uint32 VectorVMSerializeSetError_(FVectorVMSerializeState *SerializeState, uint32 Flags, uint32 LineNum)
+{
+	SerializeState->Error.Line = LineNum;
+	if (SerializeState->Error.CallbackFn)
+	{
+		SerializeState->Error.Flags = SerializeState->Error.CallbackFn(SerializeState, SerializeState->Error.Flags | Flags);
+	}
+	else
+	{
+		SerializeState->Error.Flags |= Flags;
+	}
+	if (SerializeState->Error.Flags & VVMSerErr_Fatal)
+	{
+		check(false); //hit the debugger
+		FreeVectorVMSerializeState(SerializeState);
+	}
+	return SerializeState->Error.Flags;
+}
+
+#define VectorVMSerializeSetError(SerializeState, Flags)     VectorVMSerializeSetError_(SerializeState, Flags, __LINE__)
+
+
+#if VECTORVM_SUPPORTS_EXPERIMENTAL
+
+/*********************************************************************************************************************************************************************************************************************************
+*** Serialization stuff for the Experimental VM
+*********************************************************************************************************************************************************************************************************************************/
+#define VVM_OP_XM(OpCode, Cat, NumInputs, NumOutputs, Prefix, UsageType, ...) NumInputs,
+static const uint8 VVM_OP_NUM_INPUTS[] = {
+	VVM_OP_XM_LIST
+};
+#undef VVM_OP_XM
+
+#define VVM_OP_XM(OpCode, Cat, NumInputs, NumOutputs, Prefix, UsageType, ...) NumOutputs,
+static const uint8 VVM_OP_NUM_OUTPUTS[] = {
+	VVM_OP_XM_LIST
+};
+#undef VVM_OP_XM
+
+#define VVM_OP_XM(OpCode, Cat, NumInputs, NumOutputs, Prefix, UsageType, ...) UsageType,
+static const uint8 VVM_OP_REG_USAGE_TYPE[] = {
+	VVM_OP_XM_LIST
+};
+#undef VVM_OP_XM
+
+#define VVMSer_batchStart(...)
+#define VVMSer_batchEnd(...)
+
+#define VVMSer_chunkStart(...)
+#define VVMSer_chunkEnd(...)
+
+#define VVMSer_insStart(...)
+#define VVMSer_insEndDecode(...)
+#define VVMSer_insEnd(...)
+
+
+#define VVMSer_batchStartExp(SerializeState, ...)
+#define VVMSer_batchEndExp(SerializeState, ...)
+
+#ifdef VVM_SERIALIZE_NO_WRITE
+#define VVMSer_chunkStartExp(...)
+#define VVMSer_chunkEndExp(...)
+#define VVMSer_insStartExp(...)
+#define VVMSer_insEndExp(...)
+#define VVMSer_initSerializationState(...)
+#define VVMSer_instruction(...)
+#else //VVM_SERIALIZE_NO_WRITE
+
+#define VVMIsRegIdxTempReg(Idx) ((Idx) < ExecCtx->VVMState->NumTempRegisters)
+#define VVMIsRegIdxConst(Idx)   ((Idx) >= ExecCtx->VVMState->NumTempRegisters && (Idx) < ExecCtx->VVMState->NumTempRegisters + ExecCtx->VVMState->NumConstsRemapped)
+#define VVMIsRegIdxInput(Idx)   ((Idx) >= ExecCtx->VVMState->NumTempRegisters + ExecCtx->VVMState->NumConstsRemapped)
+
+static bool VVM_areEqual32(float x, float y) {
+	float max_val = 1.f;
+	float fx = FMath::Abs(x);
+	float fy = FMath::Abs(y);
+	if (fx > max_val) {
+		max_val = fx;
+	}
+	if (fy > max_val) {
+		max_val = fy;
+	}
+	return FMath::Abs(x - y) < FLT_EPSILON * max_val;
+}
+
+void VVMSer_serializeInstruction(FVectorVMSerializeState *SerializeState, FVectorVMSerializeState *CmpSerializeState, 
+                                 FVectorVMState *VVMState, FVectorVMBatchState *BatchState, 
+                                 int StartInstance, int NumInstancesThisChunk, int NumLoops, 
+                                 int OpStart, int NumOps, uint64 StartInsCycles, uint64 EndInsCycles)
+{
+	FVectorVMSerializeInstruction *Ins = VVMSerGetNextInstruction(SerializeState, SerializeState->Running.NumChunks - 1, OpStart, NumOps, EndInsCycles - StartInsCycles, 0);
+	if (Ins != nullptr)
+	{
+		check(Ins->RegisterTable != nullptr);
+		//save the state of the temp and output registers for this instruction
+		for (uint32 i = 0; i < SerializeState->NumTempRegisters; ++i)
+		{
+			FMemory::Memcpy(Ins->RegisterTable + SerializeState->NumInstances * i + StartInstance, 
+							BatchState->RegisterData + i * NumLoops,
+							sizeof(uint32) * NumInstancesThisChunk);
+			
+		}
+		for (uint32 i = 0; i < SerializeState->NumOutputBuffers; ++i)
+		{
+			FMemory::Memcpy(Ins->RegisterTable + SerializeState->NumInstances * (SerializeState->NumTempRegisters + i) + StartInstance, 
+							BatchState->RegPtrTable[VVMState->NumTempRegisters + VVMState->NumConstBuffers + VVMState->NumInputBuffers * 2 + i], 
+							sizeof(uint32) * NumInstancesThisChunk);
+			
+		}
+		if (SerializeState->Running.NumChunks == 0)
+		{
+			//mark the registers as used float or int
+			int OpCode = (VVMState->Bytecode[Ins->OpStart - 1]);
+			if (OpCode >= 0 && OpCode < (int)EVectorVMOp::NumOpcodes && VVM_OP_CATEGORIES[OpCode] == EVectorVMOpCategory::Op)
+			{
+				uint16 *RegPtr = (uint16 *)(SerializeState->Bytecode + Ins->OpStart);
+				for (int i = 0; i < VVM_OP_NUM_INPUTS[OpCode] + VVM_OP_NUM_OUTPUTS[OpCode]; ++i) {
+					uint32 TypeFlag   = VVM_OP_REG_USAGE_TYPE[OpCode] & (1 << i) ? VVMSerIns_Int : VVMSerIns_Float;
+					uint32 TypeUnFlag = VVM_OP_REG_USAGE_TYPE[OpCode] & (1 << i) ? ~VVMSerIns_Float : ~VVMSerIns_Int;
+					uint16 AbsRegIdx;
+					int RegType = VVMGetRegisterType(VVMState, RegPtr[i], &AbsRegIdx);
+					if (RegType == VVM_RT_TEMPREG)
+					{
+						SerializeState->RegisterTableFlags[AbsRegIdx] |= TypeFlag;
+						SerializeState->RegisterTableFlags[AbsRegIdx] &= TypeUnFlag;
+					}
+					else if (RegType == VVM_RT_OUTPUT)
+					{
+						SerializeState->RegisterTableFlags[SerializeState->NumTempRegisters + AbsRegIdx] |= TypeFlag;
+						SerializeState->RegisterTableFlags[SerializeState->NumTempRegisters + AbsRegIdx] &= TypeUnFlag;
+					}
+				}
+				
+			}
+			for (uint32 i = 0; i < SerializeState->NumRegisterTable; ++i)
+			{
+				Ins->RegisterFlags[i] = SerializeState->RegisterTableFlags[i];
+			}
+		}
+
+		//compare the registers on an instruction-by-instruction basis if the VVMSer_DiffRegsPerIns flag is set
+		if (SerializeState->MismatchFn)
+		{
+			if (CmpSerializeState && Ins->OpStart < CmpSerializeState->NumBytecodeBytes && (SerializeState->Flags & VVMSer_DiffRegsPerIns))
+			{
+				uint32 InsIdx = (uint32)(Ins - SerializeState->Instructions);
+				if (InsIdx < CmpSerializeState->NumInstructions && (CmpSerializeState->Flags & VVMSer_OptimizedBytecode))
+				{
+					//only compare instructions with identical bytecode
+					FVectorVMSerializeInstruction *CmpIns = CmpSerializeState->Instructions + InsIdx;
+					if (CmpIns->OpStart == Ins->OpStart)
+					{
+						int OpCode = (VVMState->Bytecode[Ins->OpStart - 1]);
+						if (OpCode >= 0 && OpCode < (int)EVectorVMOp::NumOpcodes && VVM_OP_CATEGORIES[OpCode] == EVectorVMOpCategory::Op)
+						{
+							uint16 *RegPtr = (uint16 *)(SerializeState->Bytecode + Ins->OpStart);
+							for (uint32 i = 0; i < SerializeState->NumTempRegisters + SerializeState->NumOutputBuffers; ++i)
+							{
+								if (SerializeState->RegisterTableFlags[i] & VVMSerIns_Float)
+								{
+									float *ThisReg = (float *)(Ins->RegisterTable + SerializeState->NumInstances * i + StartInstance);
+									float *CmpReg  = (float *)(CmpIns->RegisterTable + SerializeState->NumInstances * i + StartInstance);
+									for (uint32 j = 0; j < SerializeState->NumInstances; ++j)
+									{
+										if (ThisReg[j] != CmpReg[j])
+										{
+											SerializeState->MismatchFn(false, (EVectorVMOp)OpCode, InsIdx, i, j);
+										}
+									}
+								}
+								else if (SerializeState->RegisterTableFlags[i] & VVMSerIns_Int)
+								{
+									int *ThisReg = (int *)(Ins->RegisterTable + SerializeState->NumInstances * i + StartInstance);
+									int *CmpReg  = (int *)(CmpIns->RegisterTable + SerializeState->NumInstances * i + StartInstance);
+									for (uint32 j = 0; j < SerializeState->NumInstances; ++j)
+									{
+										if (ThisReg[j] != CmpReg[j])
+										{
+											SerializeState->MismatchFn(true, (EVectorVMOp)OpCode, InsIdx, i, j);
+										}
+									}
+								}
+							}
+						}
+					}
+				}
+			}
+		}
+	}
+}
+
+uint32 *VVMSer_getRegPtrTablePtrFromIns(FVectorVMSerializeState *SerializeState, FVectorVMSerializeInstruction *Ins, uint16 RegIdx) {
+	if (RegIdx < SerializeState->NumTempRegisters)
+	{
+		return Ins->RegisterTable + SerializeState->NumInstances * RegIdx;
+	}
+	if (RegIdx >= SerializeState->NumTempRegisters + SerializeState->NumConstBuffers + SerializeState->NumInputBuffers * 2)
+	{
+		return Ins->RegisterTable + SerializeState->NumInstances * (RegIdx - SerializeState->NumConstBuffers - SerializeState->NumInputBuffers * 2);
+	}
+	return nullptr;
+}
+
+
+uint32 VVMSer_initSerializationState_(FVectorVMSerializeState *SerializeState, FVectorVMExecContext *ExecCtx, const FVectorVMOptimizeContext *OptimizeContext, uint32 Flags)
+{
+	if (SerializeState)
+	{
+		SerializeState->Error.Flags  = 0;
+
+		SerializeState->Flags            = Flags;
+		SerializeState->NumInstances     = ExecCtx->NumInstances;
+		SerializeState->NumTempRegisters = ExecCtx->VVMState->NumTempRegisters;
+		SerializeState->NumInputBuffers  = ExecCtx->VVMState->NumInputBuffers;
+		SerializeState->NumOutputBuffers = ExecCtx->VVMState->NumOutputBuffers;
+		SerializeState->NumConstBuffers  = ExecCtx->VVMState->NumConstBuffers;
+		SerializeState->NumRegisterTable = SerializeState->NumTempRegisters + ExecCtx->VVMState->NumOutputBuffers;
+
+		SerializeState->ReallocFn        = VVMDefaultRealloc;
+		SerializeState->FreeFn           = VVMDefaultFree;
+		SerializeState->OptimizeCtx      = OptimizeContext;
+		SerializeState->OptimizerHashId  = ExecCtx->VVMState->OptimizerHashId;
+
+		if (SerializeState->OptimizeCtx) {
+			SerializeState->MaxExtFnRegisters = SerializeState->OptimizeCtx->MaxExtFnRegisters;
+			SerializeState->MaxExtFnUsed = SerializeState->OptimizeCtx->MaxExtFnUsed >= 0 ? (uint32)SerializeState->OptimizeCtx->MaxExtFnUsed : 0;
+		} else {
+			SerializeState->MaxExtFnRegisters = 0;
+			SerializeState->MaxExtFnUsed = 0;
+		}
+
+		SerializeState->RegisterTableFlags = (uint8 *)SerializeState->ReallocFn(nullptr, SerializeState->NumRegisterTable, __FILE__, __LINE__);
+		if (SerializeState->RegisterTableFlags == nullptr)
+		{
+			return VectorVMSerializeSetError(SerializeState, VVMSerErr_OutOfMemory | VVMSerErr_Init | VVMSerErr_Fatal);
+		}
+		for (uint32 i = 0; i < SerializeState->NumRegisterTable; ++i)
+		{
+			SerializeState->RegisterTableFlags[i] = 0;
+		}
+
+		SerializeState->NumChunks = ExecCtx->Internal.MaxChunksPerBatch;
+		SerializeState->Chunks = (FVectorVMSerializeChunk *)SerializeState->ReallocFn(nullptr, sizeof(FVectorVMSerializeChunk) * SerializeState->NumChunks, __FILE__, __LINE__);
+		if (SerializeState->Chunks == nullptr)
+		{
+			return VectorVMSerializeSetError(SerializeState, VVMSerErr_OutOfMemory | VVMSerErr_Init | VVMSerErr_Fatal);
+		}
+		FMemory::Memset(SerializeState->Chunks, 0, sizeof(*SerializeState->Chunks) * SerializeState->NumChunks);
+		
+		if (SerializeVectorVMInputDataSets(SerializeState, ExecCtx) != 0)
+		{
+			return SerializeState->Error.Flags;
+		}
+		
+		int NumExternalFunctions = ExecCtx->VVMState->NumExtFunctions;
+		if (NumExternalFunctions != 0)
+		{
+			SerializeState->ExternalData = (FVectorVMSerializeExternalData *)SerializeState->ReallocFn(nullptr, sizeof(FVectorVMSerializeExternalData) * NumExternalFunctions, __FILE__, __LINE__);
+			if (SerializeState->ExternalData == nullptr)
+			{
+				return VectorVMSerializeSetError(SerializeState, VVMSerErr_OutOfMemory | VVMSerErr_Init | VVMSerErr_Fatal);
+			}
+			SerializeState->NumExternalData = NumExternalFunctions;
+			for (int i = 0; i < NumExternalFunctions; ++i)
+			{
+				FVectorVMExtFunctionData *f = ExecCtx->VVMState->ExtFunctionTable + i;
+				FVectorVMSerializeExternalData *ExtData = SerializeState->ExternalData + i;
+				ExtData->Name       = nullptr;
+				ExtData->NameLen    = 0;
+				ExtData->NumInputs  = f->NumInputs;
+				ExtData->NumOutputs = f->NumOutputs;
+			}
+		}
+		else
+		{
+			SerializeState->ExternalData    = nullptr;
+			SerializeState->NumExternalData = 0;
+		}
+
+		SerializeState->Bytecode = (unsigned char *)SerializeState->ReallocFn(nullptr, ExecCtx->VVMState->NumBytecodeBytes, __FILE__, __LINE__);
+		if (SerializeState->Bytecode == nullptr)
+		{
+			return VectorVMSerializeSetError(SerializeState, VVMSerErr_OutOfMemory | VVMSerErr_Init | VVMSerErr_Fatal);
+		}
+
+		FMemory::Memcpy(SerializeState->Bytecode, ExecCtx->VVMState->Bytecode, ExecCtx->VVMState->NumBytecodeBytes);
+		SerializeState->NumBytecodeBytes = ExecCtx->VVMState->NumBytecodeBytes;
+
+
+		if (SerializeState->NumOutputBuffers > 0) {
+			SerializeState->OutputRemapDataSetIdx = (uint8  *)SerializeState->ReallocFn(nullptr, sizeof(uint8) * ExecCtx->VVMState->NumOutputBuffers, __FILE__, __LINE__);
+			if (SerializeState->OutputRemapDataSetIdx == nullptr)
+			{
+				return VectorVMSerializeSetError(SerializeState, VVMSerErr_OutOfMemory | VVMSerErr_Init | VVMSerErr_Fatal);
+			}
+			SerializeState->OutputRemapDataType = (uint16 *)SerializeState->ReallocFn(nullptr, sizeof(uint16) * ExecCtx->VVMState->NumOutputBuffers, __FILE__, __LINE__);
+			if (SerializeState->OutputRemapDataType == nullptr)
+			{
+				return VectorVMSerializeSetError(SerializeState, VVMSerErr_OutOfMemory | VVMSerErr_Init | VVMSerErr_Fatal);
+			}
+			SerializeState->OutputRemapDst      = (uint16 *)SerializeState->ReallocFn(nullptr, sizeof(uint16) * ExecCtx->VVMState->NumOutputBuffers, __FILE__, __LINE__);
+			if (SerializeState->OutputRemapDst == nullptr)
+			{
+				return VectorVMSerializeSetError(SerializeState, VVMSerErr_OutOfMemory | VVMSerErr_Init | VVMSerErr_Fatal);
+			}
+
+			FMemory::Memcpy(SerializeState->OutputRemapDataSetIdx, ExecCtx->VVMState->OutputRemapDataSetIdx , sizeof(uint8 ) * ExecCtx->VVMState->NumOutputBuffers);
+			FMemory::Memcpy(SerializeState->OutputRemapDataType  , ExecCtx->VVMState->OutputRemapDataType   , sizeof(uint16) * ExecCtx->VVMState->NumOutputBuffers);
+			FMemory::Memcpy(SerializeState->OutputRemapDst       , ExecCtx->VVMState->OutputRemapDst        , sizeof(uint16) * ExecCtx->VVMState->NumOutputBuffers);
+		} else {
+			SerializeState->OutputRemapDataSetIdx = nullptr;
+			SerializeState->OutputRemapDataType   = nullptr;
+			SerializeState->OutputRemapDst        = nullptr;
+		}
+
+		
+	}
+	return 0;
+}
+#endif //VVM_SERIALIZE_NO_WRITE
+
+#else // VECTORVM_SUPPORTS_EXPERIMENTAL
+
+/*********************************************************************************************************************************************************************************************************************************
+*** Serialization stuff for the old VM
+*********************************************************************************************************************************************************************************************************************************/
+#define VVMSer_batchStartExp(...)
+#define VVMSer_batchEndExp(...)
+
+#define VVMSer_chunkStartExp(...)
+#define VVMSer_chunkEndExp(...)
+
+#define VVMSer_insStartExp(...)
+#define VVMSer_insEndExp(...)
+
+#ifdef VVM_SERIALIZE_NO_WRITE
+#define VVMSer_chunkStart(...)
+#define VVMSer_chunkEnd(...)
+#define VVMSer_insStart(...)
+#define VVMSer_insEnd(...)
+#else
+#define VVMSer_chunkStart(Context, ChunkIdx, BatchIdx)                      \
+	FVectorVMSerializeChunk *SerializeChunk = nullptr;                      \
+	int VVMSerNumInstructionsThisChunk = 0;                                 \
+	unsigned char *VVMSerStartCtxCode = (unsigned char*)Context.Code;       \
+	if (SerializeState && SerializeState->Chunks)                           \
+	{                                                                       \
+		SerializeChunk = SerializeState->Chunks + ChunkIdx;                 \
+		SerializeChunk->ChunkIdx = ChunkIdx;                                \
+		SerializeChunk->BatchIdx = BatchIdx;                                \
+		SerializeChunk->ExecIdx  = BatchIdx;                                \
+		SerializeChunk->NumInstances = NumInstancesThisChunk;               \
+		SerializeChunk->StartThreadID = FPlatformTLS::GetCurrentThreadId(); \
+		SerializeChunk->StartClock = FPlatformTime::Cycles64();             \
+	}
+
+#define VVMSer_chunkEnd(SerializeState, ...)																            \
+	if (SerializeChunk)																						            \
+	{																										            \
+		FPlatformAtomics::InterlockedOr(&SerializeState->ChunkComplete, (1ULL << (uint64)SerializeChunk->ChunkIdx));    \
+		SerializeChunk->EndThreadID = FPlatformTLS::GetCurrentThreadId();										        \
+		SerializeChunk->EndClock = FPlatformTime::Cycles64();												            \
+	}
+
+#define VVMSer_insStart(Context, ...)										\
+	uint64 VVMSerStartCycles = FPlatformTime::Cycles64();					\
+	unsigned char *VVMSerCtxStartInsCode = (unsigned char*)Context.Code;	\
+
+#define VVMSer_insEnd(Context, OpStart, NumOps, ...)																																			\
+	if (SerializeState && Op != EVectorVMOp::done)                                                                                                                                              \
+	{																																												            \
+		uint64 VVMSerEndExecCycles = FPlatformTime::Cycles64();																																	\
+		VVMSer_serializeInstruction(SerializeState, Context, StartInstance, NumInstancesThisChunk, NumLoops, OpStart, NumOps, VVMSerStartCycles, VVMSerEndExecCycles, VVMSerEndExecCycles);		\
+		uint64 VVMSerEndSerializeCycles = FPlatformTime::Cycles64();																															\
+		SerializeState->SerializeDt += VVMSerEndSerializeCycles - VVMSerEndExecCycles;																											\
+		SerializeChunk->InsExecTime += VVMSerEndSerializeCycles - VVMSerEndExecCycles;																											\
+	}
+
+
+static void VVMSer_serializeInstruction(FVectorVMSerializeState *SerializeState, FVectorVMContext &Context, 
+										int StartInstance, int NumInstancesThisChunk, int NumLoops, 
+										int OpStart, int NumOps, uint64 StartInsCycles, uint64 EndDecodeCycles, uint64 EndInsCycles)
+{
+	FVectorVMSerializeInstruction *Ins = VVMSerGetNextInstruction(SerializeState, InstructionIdx, GlobalChunkIdx, OpStart, NumOps, EndInsCycles - StartInsCycles, EndDecodeCycles - StartInsCycles);
+	if (Ins)
+	{
+		for (uint32 i = 0; i < SerializeState->NumTempRegisters; ++i)
+		{
+			unsigned char *TempReg = (unsigned char *)Context.GetTempRegister(i);
+			FMemory::Memcpy(Ins->RegisterTable + SerializeState->NumInstances * i + StartInstance, TempReg, sizeof(uint32) * NumInstancesThisChunk);
+		}
+	}
+}
+#endif //VVM_SERIALIZE_NO_WRITE
+
+#endif // VECTORVM_SUPPORTS_EXPERIMENTAL
+
+/*********************************************************************************************************************************************************************************************************************************
+*** Stuff shared between the new and old VM
+*********************************************************************************************************************************************************************************************************************************/
+
+static FVectorVMSerializeInstruction *VVMSerGetNextInstruction(FVectorVMSerializeState *SerializeState, int GlobalChunkIdx, int OpStart, int NumOps, uint64 Dt, uint64 DtDecode)
+{
+	FVectorVMSerializeInstruction *Ins = nullptr;
+	int InstructionIdx = SerializeState->Running.NumInstructionsThisChunk++;
+	if (GlobalChunkIdx == 0)
+	{
+		if (SerializeState->NumInstructions + 1 >= SerializeState->NumInstructionsAllocated)
+		{
+			if (SerializeState->NumInstructionsAllocated == 0)
+			{
+				SerializeState->NumInstructionsAllocated = 1024;
+				check(SerializeState->Instructions == nullptr);
+			}
+			else
+			{
+				SerializeState->NumInstructionsAllocated <<= 1;
+			}
+			FVectorVMSerializeInstruction *NewInstructions = (FVectorVMSerializeInstruction *)SerializeState->ReallocFn(SerializeState->Instructions, sizeof(FVectorVMSerializeInstruction) * SerializeState->NumInstructionsAllocated, __FILE__, __LINE__);
+			if (NewInstructions == nullptr)
+			{
+				VectorVMSerializeSetError(SerializeState, VVMSerErr_OutOfMemory | VVMSerErr_Instruction | VVMSerErr_Fatal);
+				return nullptr;
+			}
+			else
+			{
+				SerializeState->Instructions = NewInstructions;
+			}
+		}
+		Ins = SerializeState->Instructions + SerializeState->NumInstructions++;
+
+		//we only store this information on the first chunk
+		Ins->OpStart       = OpStart;
+		Ins->NumOps        = NumOps;
+		Ins->Dt            = Dt;
+		Ins->DtDecode      = DtDecode;
+
+		Ins->RegisterTable = (uint32 *)SerializeState->ReallocFn(nullptr, sizeof(uint32) * SerializeState->NumInstances * SerializeState->NumRegisterTable, __FILE__, __LINE__); //alloc for *ALL* chunks
+		if (Ins->RegisterTable == nullptr)
+		{
+			VectorVMSerializeSetError(SerializeState, VVMSerErr_OutOfMemory | VVMSerErr_Instruction | VVMSerErr_Fatal);
+			return nullptr;
+		}
+		Ins->RegisterFlags = (unsigned char *)SerializeState->ReallocFn(nullptr, SerializeState->NumRegisterTable, __FILE__, __LINE__);
+		if (Ins->RegisterFlags == nullptr)
+		{
+			VectorVMSerializeSetError(SerializeState, VVMSerErr_OutOfMemory | VVMSerErr_Instruction | VVMSerErr_Fatal);
+			return nullptr;
+		}
+		FMemory::Memset(Ins->RegisterFlags, 0, SerializeState->NumRegisterTable);
+	}
+	else
+	{
+		if ((SerializeState->ChunkComplete & 1) == 0)
+		{
+			//we can't serialize the instructions until the first chunk is done since it's the one that allocates the memory.
+			//if this ever becomes a bottleneck (who cares for serialization!?), but if it ever does I can the change the 
+			//memory allocation such that each chunk's temp registers for each instruction gets grouped in per-chunk allocations
+			//then the debugger can merge them all into a single array on load... but this is simpler and I doubt will become an issue
+			uint32 SanityCount = 0;
+			while ((SerializeState->ChunkComplete & 1) == 0 && SanityCount++ < (1ULL << 31))
+			{
+				FPlatformProcess::Yield();
+			}
+			check(SanityCount < (1ULL << 31) - 1);
+		}
+		check(InstructionIdx < (int)SerializeState->NumInstructions);
+		Ins = SerializeState->Instructions + InstructionIdx;
+	}
+	return Ins;
+}
+
+#ifndef VVM_SERIALIZE_NO_IO
+void SerializeVectorVMWriteToFile(FVectorVMSerializeState *SerializeState, uint8 WhichStateWritten, const wchar_t *Filename)
+{
+	IPlatformFile &PlatformFile = FPlatformFileManager::Get().GetPlatformFile();
+	IFileHandle *File           = PlatformFile.OpenWrite((TCHAR *)Filename);
+	if (SerializeState->OptimizeCtx)
+	{
+		SerializeState->Flags |= VVMSer_IncludeOptContext;
+	}
+	else
+	{
+		SerializeState->Flags &= ~VVMSer_IncludeOptContext;
+	}
+	if (File)
+	{
+		File->Write(&WhichStateWritten, 1);		//1 = Exp, 2 = UE.  OR for both (1|2 = 3)
+		File->Write((uint8 *)&SerializeState->NumInstances     , sizeof(uint32));
+		File->Write((uint8 *)&SerializeState->Flags            , sizeof(uint32));
+		File->Write((uint8 *)&SerializeState->ExecDt           , sizeof(uint64));
+		File->Write((uint8 *)&SerializeState->SerializeDt      , sizeof(uint64));
+		File->Write((uint8 *)&SerializeState->NumTempRegisters , sizeof(uint32));
+		File->Write((uint8 *)&SerializeState->NumInputBuffers  , sizeof(uint32));
+		File->Write((uint8 *)&SerializeState->NumOutputBuffers , sizeof(uint32));
+		File->Write((uint8 *)&SerializeState->NumConstBuffers  , sizeof(uint32));
+		File->Write((uint8 *)&SerializeState->NumRegisterTable , sizeof(uint32));
+		File->Write((uint8 *)&SerializeState->NumBytecodeBytes , sizeof(uint32));
+		File->Write((uint8 *)&SerializeState->NumInstructions  , sizeof(uint32));
+		File->Write((uint8 *)&SerializeState->NumDataSets      , sizeof(uint32));
+		File->Write((uint8 *)&SerializeState->NumChunks        , sizeof(uint32));
+		File->Write((uint8 *)&SerializeState->NumExternalData  , sizeof(uint32));
+		File->Write((uint8 *)&SerializeState->MaxExtFnRegisters, sizeof(uint32));
+		File->Write((uint8 *)&SerializeState->MaxExtFnUsed     , sizeof(uint32));
+		File->Write((uint8 *)&SerializeState->OptimizerHashId  , sizeof(uint64));
+
+		
+		//write the datasets
+		for (uint32 i = 0; i < SerializeState->NumDataSets; ++i)
+		{
+			FVectorVMSerializeDataSet *DataSet = SerializeState->DataSets + i;
+			uint32 IOFlag = (DataSet->InputBuffers ? 1 : 0) + (DataSet->OutputBuffers ? 2 : 0);
+			File->Write((uint8 *)&IOFlag                            , sizeof(uint32));
+			File->Write((uint8 *)&DataSet->InputOffset[0]           , sizeof(uint32));
+			File->Write((uint8 *)&DataSet->InputOffset[1]           , sizeof(uint32));
+			File->Write((uint8 *)&DataSet->InputOffset[2]           , sizeof(uint32));
+			File->Write((uint8 *)&DataSet->InputOffset[3]           , sizeof(uint32));
+			File->Write((uint8 *)&DataSet->OutputOffset[0]          , sizeof(uint32));
+			File->Write((uint8 *)&DataSet->OutputOffset[1]          , sizeof(uint32));
+			File->Write((uint8 *)&DataSet->OutputOffset[2]          , sizeof(uint32));
+			File->Write((uint8 *)&DataSet->OutputOffset[3]          , sizeof(uint32));
+
+			File->Write((uint8 *)&DataSet->InputInstanceOffset      , sizeof(int32));
+			File->Write((uint8 *)&DataSet->InputDataSetAccessIndex  , sizeof(int32));
+			File->Write((uint8 *)&DataSet->InputIDAcquireTag        , sizeof(int32));
+			File->Write((uint8 *)&DataSet->InputNumFreeIDs          , sizeof(int32));
+			File->Write((uint8 *)&DataSet->InputMaxUsedIDs          , sizeof(int32));
+			File->Write((uint8 *)&DataSet->InputNumSpawnedIDs       , sizeof(int32));
+
+			File->Write((uint8 *)&DataSet->InputIDTableNum          , sizeof(int32));
+			File->Write((uint8 *)&DataSet->InputFreeIDTableNum      , sizeof(int32));
+			File->Write((uint8 *)&DataSet->InputSpawnedIDTableNum   , sizeof(int32));
+
+			File->Write((uint8 *)&DataSet->OutputInstanceOffset     , sizeof(int32));
+			File->Write((uint8 *)&DataSet->OutputDataSetAccessIndex , sizeof(int32));
+			File->Write((uint8 *)&DataSet->OutputIDAcquireTag       , sizeof(int32));
+			File->Write((uint8 *)&DataSet->OutputNumFreeIDs         , sizeof(int32));
+			File->Write((uint8 *)&DataSet->OutputMaxUsedIDs         , sizeof(int32));
+			File->Write((uint8 *)&DataSet->OutputNumSpawnedIDs      , sizeof(int32));
+
+			File->Write((uint8 *)&DataSet->OutputIDTableNum         , sizeof(int32));
+			File->Write((uint8 *)&DataSet->OutputFreeIDTableNum     , sizeof(int32));
+			File->Write((uint8 *)&DataSet->OutputSpawnedIDTableNum  , sizeof(int32));
+
+			uint32 TotalNumInstances = SerializeState->NumInstances + DataSet->InputInstanceOffset; //the VM could run 100 instances, but the dataset could have 150, and we start at 50, so just write all instances in the data set
+			//write input bufers
+			if (DataSet->InputBuffers)
+			{
+				uint32 Num32BitBuffers = DataSet->InputOffset[2];
+				uint32 Num16BitBuffers = DataSet->InputOffset[3] - Num32BitBuffers;
+				File->Write((uint8 *)DataSet->InputBuffers                                                       , sizeof(uint32) * Num32BitBuffers * TotalNumInstances);
+				File->Write((uint8 *)DataSet->InputBuffers + sizeof(uint32) * Num32BitBuffers * TotalNumInstances, sizeof(uint16) * Num16BitBuffers * TotalNumInstances);
+			}
+			if (DataSet->OutputBuffers)
+			{
+				uint32 Num32BitBuffers = DataSet->OutputOffset[2];
+				uint32 Num16BitBuffers = DataSet->OutputOffset[3] - Num32BitBuffers;
+				File->Write((uint8 *)DataSet->OutputBuffers                                                        , sizeof(uint32) * Num32BitBuffers * TotalNumInstances);
+				File->Write((uint8 *)DataSet->OutputBuffers + sizeof(uint32) * Num32BitBuffers * TotalNumInstances , sizeof(uint16) * Num16BitBuffers * TotalNumInstances);
+			}
+			//input tables
+			if (DataSet->InputIDTableNum > 0)
+			{
+				File->Write((uint8 *)DataSet->InputIDTable          , sizeof(int32) * DataSet->InputIDTableNum);
+			}
+			if (DataSet->InputFreeIDTableNum > 0)
+			{
+				File->Write((uint8 *)DataSet->InputFreeIDTable      , sizeof(int32) * DataSet->InputFreeIDTableNum);
+			}
+			if (DataSet->InputSpawnedIDTableNum > 0)
+			{
+				File->Write((uint8 *)DataSet->InputSpawnedIDTable   , sizeof(int32) * DataSet->InputSpawnedIDTableNum);
+			}
+			//output tables
+			if (DataSet->OutputIDTableNum > 0)
+			{
+				File->Write((uint8 *)DataSet->OutputIDTable         , sizeof(int32) * DataSet->OutputIDTableNum);
+			}
+			if (DataSet->OutputFreeIDTableNum > 0)
+			{
+				File->Write((uint8 *)DataSet->OutputFreeIDTable     , sizeof(int32) * DataSet->OutputFreeIDTableNum);
+			}
+			if (DataSet->OutputSpawnedIDTableNum > 0)
+			{
+				File->Write((uint8 *)DataSet->OutputSpawnedIDTable  , sizeof(int32) * DataSet->OutputSpawnedIDTableNum);
+			}
+		}
+
+		//write the chunk info
+		for (uint32 i = 0; i < SerializeState->NumChunks; ++i)
+		{
+			FVectorVMSerializeChunk *Chunk = SerializeState->Chunks + i;
+			File->Write((uint8 *)&Chunk->ChunkIdx		, sizeof(uint32));
+			File->Write((uint8 *)&Chunk->BatchIdx		, sizeof(uint32));
+			File->Write((uint8 *)&Chunk->NumInstances	, sizeof(uint32));
+			File->Write((uint8 *)&Chunk->StartThreadID	, sizeof(uint32));
+			File->Write((uint8 *)&Chunk->EndThreadID	, sizeof(uint32));
+			File->Write((uint8 *)&Chunk->StartClock	    , sizeof(uint64));
+			File->Write((uint8 *)&Chunk->EndClock		, sizeof(uint64));
+		}
+		
+		{ //write const data
+			int ConstTableSizesInBytes = 0;
+			for (uint32 i = 0; i < SerializeState->NumConstBuffers; ++i) {
+				ConstTableSizesInBytes += SerializeState->ConstTableSizesInBytes[i];
+			}
+			File->Write((uint8 *)SerializeState->ConstTableSizesInBytes, sizeof(uint32) * SerializeState->NumConstBuffers);
+			File->Write((uint8 *)SerializeState->PreExecConstData      , ConstTableSizesInBytes);
+			File->Write((uint8 *)SerializeState->PostExecConstData     , ConstTableSizesInBytes);
+		}
+
+		if (SerializeState->NumOutputBuffers)
+		{
+			File->Write((uint8 *)SerializeState->OutputRemapDataSetIdx, sizeof(uint8 ) * SerializeState->NumOutputBuffers);
+			File->Write((uint8 *)SerializeState->OutputRemapDataType  , sizeof(uint16) * SerializeState->NumOutputBuffers);
+			File->Write((uint8 *)SerializeState->OutputRemapDst       , sizeof(uint16) * SerializeState->NumOutputBuffers);
+		}
+
+		//write bytecode
+		if (SerializeState->NumBytecodeBytes > 0)
+		{
+			File->Write((uint8 *)SerializeState->Bytecode, SerializeState->NumBytecodeBytes);
+		}
+
+		//write instructions
+		for (uint32 i = 0; i < SerializeState->NumInstructions; ++i)
+		{
+			FVectorVMSerializeInstruction *Ins = SerializeState->Instructions + i;
+			File->Write((uint8 *)&Ins->OpStart      , sizeof(uint32));
+			File->Write((uint8 *)&Ins->NumOps       , sizeof(uint32));
+			File->Write((uint8 *)&Ins->Dt           , sizeof(uint64));
+			File->Write((uint8 *)&Ins->DtDecode     , sizeof(uint64));
+			File->Write((uint8 *)Ins->RegisterTable , sizeof(uint32) * SerializeState->NumInstances * SerializeState->NumRegisterTable);
+			if (SerializeState->NumBytecodeBytes)
+			{
+				File->Write((uint8 *)Ins->RegisterFlags, SerializeState->NumRegisterTable);
+			}
+		}
+
+		//write external function data
+		for (uint32 i = 0; i < SerializeState->NumExternalData; ++i)
+		{
+			FVectorVMSerializeExternalData *ExtData = SerializeState->ExternalData + i;
+			File->Write((uint8 *)&ExtData->NameLen    , sizeof(uint16));
+			File->Write((uint8 *)ExtData->Name        , sizeof(uint16) * ExtData->NameLen);
+			File->Write((uint8 *)&ExtData->NumInputs  , sizeof(uint16));
+			File->Write((uint8 *)&ExtData->NumOutputs , sizeof(uint16));
+		}
+#if VECTORVM_SUPPORTS_EXPERIMENTAL
+		//write the optimization context
+		if (SerializeState->Flags & VVMSer_IncludeOptContext) {
+			check(SerializeState->OptimizeCtx);
+			const FVectorVMOptimizeContext *OptCtx = SerializeState->OptimizeCtx;
+			File->Write((uint8 *)&OptCtx->MaxOutputDataSet                  , sizeof(uint32));
+			File->Write((uint8 *)&OptCtx->NumConstsAlloced                  , sizeof(uint16));
+			File->Write((uint8 *)&OptCtx->NumTempRegisters					, sizeof(uint32));
+			File->Write((uint8 *)&OptCtx->NumConstsRemapped					, sizeof(uint16));
+			File->Write((uint8 *)&OptCtx->NumInputsRemapped					, sizeof(uint16));
+			File->Write((uint8 *)&OptCtx->NumOutputsAliasedToTempRegisters	, sizeof(uint16));
+			File->Write((uint8 *)&OptCtx->NumNoAdvanceInputs				, sizeof(uint16));
+			File->Write((uint8 *)&OptCtx->NumInputDataSets					, sizeof(uint16));
+			File->Write((uint8 *)&OptCtx->NumOutputsRemapped				, sizeof(uint16));
+			File->Write((uint8 *)&OptCtx->NumOutputInstructions				, sizeof(uint16));
+			File->Write((uint8 *)&OptCtx->NumExtFns							, sizeof(uint32));
+			File->Write((uint8 *)&OptCtx->MaxExtFnRegisters					, sizeof(uint32));
+			File->Write((uint8 *)&OptCtx->NumDummyRegsReq					, sizeof(uint32));
+			File->Write((uint8 *)&OptCtx->MaxExtFnUsed						, sizeof(int32));
+			File->Write((uint8 *)&OptCtx->Flags								, sizeof(uint32));
+			File->Write((uint8 *)OptCtx->ConstRemap[1]        , sizeof(uint16) * OptCtx->NumConstsRemapped);
+			File->Write((uint8 *)OptCtx->InputRemapTable      , sizeof(uint16) * OptCtx->NumInputsRemapped);
+			File->Write((uint8 *)OptCtx->OutputRemapDataSetIdx, sizeof(uint8)  * OptCtx->NumOutputsRemapped);
+			File->Write((uint8 *)OptCtx->OutputRemapDataType  , sizeof(uint16) * OptCtx->NumOutputsRemapped);
+			File->Write((uint8 *)OptCtx->OutputRemapDst       , sizeof(uint16) * OptCtx->NumOutputsRemapped);
+			if (OptCtx->NumInputDataSets) {
+				File->Write((uint8 *)OptCtx->InputDataSetOffsets      , sizeof(uint16) * 8 * OptCtx->NumInputDataSets);
+			}
+		}
+		delete File;
+#endif //VECTORVM_SUPPORTS_EXPERIMENTAL
+	}
+}
+
+#else
+
+void SerializeVectorVMWriteToFile(FVectorVMSerializeState *SerializeState, uint8 WhichStateWritten, const wchar_t *Filename)
+{
+
+}
+
+#endif //VVM_SERIALIZE_NO_IO
+
+
+static uint32 SerializeConstData(FVectorVMSerializeState *SerializeState, const uint8 * const *ConstantTableData, const int *ConstantTableSizes, int32 ConstantTableCount, uint32 **OutConstData, uint32 **OutConstTableSizes)
+{
+	int TotalConstNumBytes = 0;
+	for (int i = 0; i < ConstantTableCount; ++i)
+	{
+		TotalConstNumBytes += ConstantTableSizes[i];
+	}
+	*OutConstData = (uint32 *)SerializeState->ReallocFn(nullptr, TotalConstNumBytes, __FILE__, __LINE__);
+	if (*OutConstData == nullptr)
+	{
+		VectorVMSerializeSetError(SerializeState, VVMSerErr_OutOfMemory | VVMSerErr_ConstData | VVMSerErr_Fatal);
+		return 0;
+	}
+	if (OutConstTableSizes) {
+		*OutConstTableSizes = (uint32 *)SerializeState->ReallocFn(nullptr, sizeof(uint32) * ConstantTableCount, __FILE__, __LINE__);
+		if (*OutConstTableSizes == nullptr)
+		{
+			VectorVMSerializeSetError(SerializeState, VVMSerErr_OutOfMemory | VVMSerErr_ConstData | VVMSerErr_Fatal);
+			return 0;
+		}
+		for (int i = 0; i < ConstantTableCount; ++i)
+		{
+			(*OutConstTableSizes)[i] = ConstantTableSizes[i];
+		}
+	}
+	uint8 *CPtr = (uint8 *)*OutConstData;
+	for (int i = 0; i < ConstantTableCount; ++i)
+	{
+		FMemory::Memcpy(CPtr, ConstantTableData[i], ConstantTableSizes[i]);
+		CPtr += ConstantTableSizes[i];
+	}
+	return ConstantTableCount;
+}
+
+uint32 SerializeVectorVMInputDataSets(FVectorVMSerializeState *SerializeState, TArrayView<FDataSetMeta> DataSets, const uint8 * const *ConstantTableData, const int *ConstantTableSizes, int32 ConstantTableCount)
+{
+	SerializeState->NumDataSets = DataSets.Num();
+	SerializeState->DataSets = (FVectorVMSerializeDataSet *)SerializeState->ReallocFn(nullptr, sizeof(FVectorVMSerializeDataSet) * SerializeState->NumDataSets, __FILE__, __LINE__);
+	if (SerializeState->DataSets == nullptr)
+	{
+		return VectorVMSerializeSetError(SerializeState, VVMSerErr_OutOfMemory | VVMSerErr_InputDataSets | VVMSerErr_Fatal);
+	}
+	FMemory::Memset(SerializeState->DataSets, 0, sizeof(FVectorVMSerializeDataSet) * SerializeState->NumDataSets);
+	for (uint32 i = 0; i < SerializeState->NumDataSets; ++i)
+	{
+		FVectorVMSerializeDataSet *DsDst = SerializeState->DataSets + i;
+		FDataSetMeta *DsSrc = &DataSets[i];
+		DsDst->InputInstanceOffset      = DsSrc->InstanceOffset;
+		DsDst->InputDataSetAccessIndex  = DsSrc->DataSetAccessIndex;
+		DsDst->InputIDAcquireTag        = DsSrc->IDAcquireTag;
+		DsDst->InputNumFreeIDs          = *DsSrc->NumFreeIDs;
+		DsDst->InputMaxUsedIDs          = *DsSrc->MaxUsedID;
+		DsDst->InputNumSpawnedIDs       = *DsSrc->NumSpawnedIDs;
+		
+		//we serialize the whole dataset... even if there's more instances than we use just so the offsets are correct when we sim
+		uint32 TotalNumInstances        = SerializeState->NumInstances + DsDst->InputInstanceOffset;
+		uint32 **SrcInputBuffers        = (uint32 **)DsSrc->InputRegisters.GetData();
+		if (SrcInputBuffers)
+		{
+			DsDst->InputOffset[0] = DsSrc->InputRegisterTypeOffsets[0];
+			DsDst->InputOffset[1] = DsSrc->InputRegisterTypeOffsets[1];
+			DsDst->InputOffset[2] = DsSrc->InputRegisterTypeOffsets[2];
+			DsDst->InputOffset[3] = DsSrc->InputRegisters.Num();
+			uint32 TotalNum32BitBuffers = DsDst->InputOffset[2];
+			uint32 TotalNum16BitBuffers = DsDst->InputOffset[3] - TotalNum32BitBuffers;
+			DsDst->InputBuffers = (uint32 *)SerializeState->ReallocFn(nullptr, (sizeof(uint32) * TotalNum32BitBuffers + sizeof(uint16) * TotalNum16BitBuffers) * TotalNumInstances, __FILE__, __LINE__);
+			if (DsDst->InputBuffers == nullptr)
+			{
+				return VectorVMSerializeSetError(SerializeState, VVMSerErr_OutOfMemory | VVMSerErr_InputDataSets | VVMSerErr_Fatal);
+			}
+			for (uint32 j = 0; j < TotalNum32BitBuffers; ++j)
+			{
+				FMemory::Memcpy(DsDst->InputBuffers + j * TotalNumInstances, SrcInputBuffers[j], sizeof(uint32) * TotalNumInstances);
+			}
+			uint16 *InputBuff16 = (uint16 *)(DsDst->InputBuffers + TotalNum32BitBuffers * TotalNumInstances);
+			for (uint32 j = 0; j < TotalNum16BitBuffers; ++j)
+			{
+				FMemory::Memcpy(InputBuff16 + j * TotalNumInstances, SrcInputBuffers[TotalNum32BitBuffers + j], sizeof(uint16) * TotalNumInstances);
+			}
+		}
+		else
+		{
+			DsDst->InputOffset[0] = 0;
+			DsDst->InputOffset[1] = 0;
+			DsDst->InputOffset[2] = 0;
+			DsDst->InputBuffers   = nullptr;
+		}
+		
+		//IDTable
+		if (DsSrc->IDTable && DsSrc->IDTable->Num() > 0)
+		{
+			DsDst->InputIDTable = (int32 *)SerializeState->ReallocFn(nullptr, sizeof(int) * DsSrc->IDTable->Num(), __FILE__, __LINE__);
+			if (DsDst->InputIDTable == nullptr)
+			{
+				return VectorVMSerializeSetError(SerializeState, VVMSerErr_OutOfMemory | VVMSerErr_InputDataSets | VVMSerErr_Fatal);
+			}
+			FMemory::Memcpy(DsDst->InputIDTable, DsSrc->IDTable->GetData(), sizeof(int32) * DsSrc->IDTable->Num());
+			DsDst->InputIDTableNum = DsSrc->IDTable->Num();
+		}
+		else
+		{
+			DsDst->InputIDTable    = nullptr;
+			DsDst->InputIDTableNum = 0;
+		}
+
+		//FreeIDTable
+		if (DsSrc->FreeIDTable && DsSrc->FreeIDTable->Num() > 0)
+		{
+			DsDst->InputFreeIDTable = (int32 *)SerializeState->ReallocFn(nullptr, sizeof(int) * DsSrc->FreeIDTable->Num(), __FILE__, __LINE__);
+			if (DsDst->InputFreeIDTable == nullptr)
+			{
+				return VectorVMSerializeSetError(SerializeState, VVMSerErr_OutOfMemory | VVMSerErr_InputDataSets | VVMSerErr_Fatal);
+			}
+			FMemory::Memcpy(DsDst->InputFreeIDTable, DsSrc->FreeIDTable->GetData(), sizeof(int32) * DsSrc->FreeIDTable->Num());
+			DsDst->InputFreeIDTableNum = DsSrc->FreeIDTable->Num();
+		}
+		else
+		{
+			DsDst->InputFreeIDTable    = nullptr;
+			DsDst->InputFreeIDTableNum = 0;
+		}
+
+		//SpawnedIDTable
+		if (DsSrc->SpawnedIDsTable && DsSrc->SpawnedIDsTable->Num() > 0)
+		{
+			DsDst->InputSpawnedIDTable = (int32 *)SerializeState->ReallocFn(nullptr, sizeof(int) * DsSrc->SpawnedIDsTable->Num(), __FILE__, __LINE__);
+			if (DsDst->InputSpawnedIDTable == nullptr)
+			{
+				return VectorVMSerializeSetError(SerializeState, VVMSerErr_OutOfMemory | VVMSerErr_InputDataSets | VVMSerErr_Fatal);
+			}
+			FMemory::Memcpy(DsDst->InputSpawnedIDTable, DsSrc->SpawnedIDsTable->GetData(), sizeof(int32) * DsSrc->SpawnedIDsTable->Num());
+			DsDst->InputSpawnedIDTableNum = DsSrc->SpawnedIDsTable->Num();
+		}
+		else
+		{
+			DsDst->InputSpawnedIDTable    = nullptr;
+			DsDst->InputSpawnedIDTableNum = 0;
+		}
+	}
+	SerializeState->NumConstBuffers = SerializeConstData(SerializeState, ConstantTableData, ConstantTableSizes, ConstantTableCount, &SerializeState->PreExecConstData, &SerializeState->ConstTableSizesInBytes);
+	return SerializeState->Error.Flags;
+}
+
+uint32 SerializeVectorVMOutputDataSets(FVectorVMSerializeState *SerializeState, TArrayView<FDataSetMeta> DataSets, const uint8 * const *ConstantTableData, const int *ConstantTableSizes, int32 ConstantTableCount)
+{
+	if (SerializeState->NumDataSets == DataSets.Num()) //must have already serialized the input data sets
+	{
+		for (int i = 0; i < DataSets.Num(); ++i)
+		{
+			FVectorVMSerializeDataSet *DsDst = SerializeState->DataSets + i;
+			FDataSetMeta *DsSrc = &DataSets[i];
+			DsDst->OutputOffset[0]			= DsSrc->OutputRegisterTypeOffsets[0];
+			DsDst->OutputOffset[1]			= DsSrc->OutputRegisterTypeOffsets[1];
+			DsDst->OutputOffset[2]			= DsSrc->OutputRegisterTypeOffsets[2];
+			DsDst->OutputOffset[3]          = DsSrc->OutputRegisters.Num();
+			DsDst->OutputInstanceOffset		= DsSrc->InstanceOffset;
+			DsDst->OutputDataSetAccessIndex	= DsSrc->DataSetAccessIndex;
+			DsDst->OutputIDAcquireTag		= DsSrc->IDAcquireTag;
+			DsDst->OutputNumFreeIDs         = *DsSrc->NumFreeIDs;
+			DsDst->OutputMaxUsedIDs         = *DsSrc->MaxUsedID;
+			DsDst->OutputNumSpawnedIDs      = *DsSrc->NumSpawnedIDs;
+
+			uint32 TotalNumInstances = SerializeState->NumInstances + DsDst->OutputInstanceOffset;
+			uint32 **SrcOutputBuffers = (uint32 **)DsSrc->OutputRegisters.GetData();
+			if (SrcOutputBuffers)
+			{
+				uint32 TotalNum32BitBuffers = DsDst->OutputOffset[2];
+				uint32 TotalNum16BitBuffers = DsDst->OutputOffset[3] - TotalNum32BitBuffers;
+				DsDst->OutputBuffers = (uint32 *)SerializeState->ReallocFn(nullptr, (sizeof(uint32) * TotalNum32BitBuffers + sizeof(uint16) * TotalNum16BitBuffers) * TotalNumInstances, __FILE__, __LINE__);
+				if (DsDst->OutputBuffers == nullptr)
+				{
+					return VectorVMSerializeSetError(SerializeState, VVMSerErr_OutOfMemory | VVMSerErr_OutputDataSets | VVMSerErr_Fatal);
+				}
+				for (uint32 j = 0; j < TotalNum32BitBuffers; ++j)
+				{
+					FMemory::Memcpy(DsDst->OutputBuffers + j * TotalNumInstances, SrcOutputBuffers[j], sizeof(uint32) * TotalNumInstances);
+				}
+				uint16 *OutputBuff16 = (uint16 *)(DsDst->OutputBuffers + TotalNum32BitBuffers * TotalNumInstances);
+				for (uint32 j = 0; j < TotalNum16BitBuffers; ++j)
+				{
+					FMemory::Memcpy(OutputBuff16 + j * TotalNumInstances, SrcOutputBuffers[TotalNum32BitBuffers + j], sizeof(uint16) * TotalNumInstances);
+				}
+			}
+			else
+			{
+				DsDst->OutputBuffers = nullptr;
+			}
+
+			//IDTable
+			if (DsSrc->IDTable && DsSrc->IDTable->Num() > 0)
+			{
+				DsDst->OutputIDTable = (int32 *)SerializeState->ReallocFn(nullptr, sizeof(int) * DsSrc->IDTable->Num(), __FILE__, __LINE__);
+				if (DsDst->OutputIDTable == nullptr)
+				{
+					return VectorVMSerializeSetError(SerializeState, VVMSerErr_OutOfMemory | VVMSerErr_OutputDataSets | VVMSerErr_Fatal);
+				}
+				FMemory::Memcpy(DsDst->OutputIDTable, DsSrc->IDTable->GetData(), sizeof(int32) * DsSrc->IDTable->Num());
+				DsDst->OutputIDTableNum = DsSrc->IDTable->Num();
+			}
+			else
+			{
+				DsDst->OutputIDTable    = nullptr;
+				DsDst->OutputIDTableNum = 0;
+			}
+
+			//FreeIDTable
+			if (DsSrc->FreeIDTable && DsSrc->FreeIDTable->Num() > 0)
+			{
+				DsDst->OutputFreeIDTable = (int32 *)SerializeState->ReallocFn(nullptr, sizeof(int) * DsSrc->FreeIDTable->Num(), __FILE__, __LINE__);
+				if (DsDst->OutputFreeIDTable == nullptr)
+				{
+					return VectorVMSerializeSetError(SerializeState, VVMSerErr_OutOfMemory | VVMSerErr_OutputDataSets | VVMSerErr_Fatal);
+				}
+				FMemory::Memcpy(DsDst->OutputFreeIDTable, DsSrc->FreeIDTable->GetData(), sizeof(int32) * DsSrc->FreeIDTable->Num());
+				DsDst->OutputFreeIDTableNum = DsSrc->FreeIDTable->Num();
+			}
+			else
+			{
+				DsDst->OutputFreeIDTable    = nullptr;
+				DsDst->OutputFreeIDTableNum = 0;
+			}
+
+			//SpawnedIDTable
+			if (DsSrc->SpawnedIDsTable && DsSrc->SpawnedIDsTable->Num() > 0)
+			{
+				DsDst->OutputSpawnedIDTable = (int32 *)SerializeState->ReallocFn(nullptr, sizeof(int) * DsSrc->SpawnedIDsTable->Num(), __FILE__, __LINE__);
+				if (DsDst->OutputSpawnedIDTable == nullptr)
+				{
+					return VectorVMSerializeSetError(SerializeState, VVMSerErr_OutOfMemory | VVMSerErr_OutputDataSets | VVMSerErr_Fatal);
+				}
+				FMemory::Memcpy(DsDst->OutputSpawnedIDTable, DsSrc->SpawnedIDsTable->GetData(), sizeof(int32) * DsSrc->SpawnedIDsTable->Num());
+				DsDst->OutputSpawnedIDTableNum = DsSrc->SpawnedIDsTable->Num();
+			}
+			else
+			{
+				DsDst->OutputSpawnedIDTable    = nullptr;
+				DsDst->OutputSpawnedIDTableNum = 0;
+			}
+
+		}
+	}
+	else
+	{
+		VectorVMSerializeSetError(SerializeState, VVMSerErr_OutputDataSets);
+		for (uint32 i = 0; i < SerializeState->NumDataSets; ++i)
+		{
+			FVectorVMSerializeDataSet *DataSet = SerializeState->DataSets + i;
+			DataSet->OutputOffset[0] = 0;
+			DataSet->OutputOffset[1] = 0;
+			DataSet->OutputOffset[2] = 0;
+			DataSet->OutputBuffers = nullptr;
+			DataSet->OutputIDTable = nullptr;
+			DataSet->OutputIDTableNum = 0;
+			DataSet->OutputFreeIDTable = nullptr;
+			DataSet->OutputFreeIDTableNum = 0;
+			DataSet->OutputSpawnedIDTable = nullptr;
+			DataSet->OutputSpawnedIDTableNum = 0;
+		}
+	}
+	if (SerializeState->Error.Flags == 0)
+	{
+		SerializeState->NumConstBuffers = SerializeConstData(SerializeState, ConstantTableData, ConstantTableSizes, ConstantTableCount, &SerializeState->PostExecConstData, nullptr);
+	}
+	return SerializeState->Error.Flags;
+}
+
+#if VECTORVM_SUPPORTS_EXPERIMENTAL
+uint32 SerializeVectorVMInputDataSets(FVectorVMSerializeState *SerializeState, FVectorVMExecContext *ExecCtx)
+{
+	return SerializeState ? SerializeVectorVMInputDataSets(SerializeState, ExecCtx->DataSets, ExecCtx->ConstantTableData, ExecCtx->ConstantTableNumBytes, ExecCtx->ConstantTableCount) : 0;
+}
+
+uint32 SerializeVectorVMOutputDataSets(FVectorVMSerializeState *SerializeState, FVectorVMExecContext *ExecCtx)
+{
+	return SerializeState ? SerializeVectorVMOutputDataSets(SerializeState, ExecCtx->DataSets, ExecCtx->ConstantTableData, ExecCtx->ConstantTableNumBytes, ExecCtx->ConstantTableCount) : 0;
+}
+#endif // #if VECTORVM_SUPPORTS_EXPERIMENTAL
+
+#else //VECTORVM_SUPPORTS_SERIALIZATION
+
+#define VVMSer_batchStart(...)
+#define VVMSer_batchEnd(...)
+
+#define VVMSer_chunkStart(...)
+#define VVMSer_chunkEnd(...)
+
+#define VVMSer_insStart(...)
+#define VVMSer_insEndDecode(...)
+#define VVMSer_insEnd(...)
+
+#define VVMSer_batchStartExp(...)
+#define VVMSer_batchEndExp(...)
+
+#define VVMSer_chunkStartExp(...)
+#define VVMSer_chunkEndExp(...)
+
+#define VVMSer_insStartExp(...)
+#define VVMSer_insEndExp(...)
+
+#define VVMSer_instruction(...)
+
+#define VVMSer_initSerializationState(...)
+
+uint64 VVMSer_cmpStates(FVectorVMSerializeState *S0, FVectorVMSerializeState *S1, uint32 Flags)
+{
+	return 0;
+}
+
+uint32 SerializeVectorVMInputDataSets  (FVectorVMSerializeState *SerializeState, struct FVectorVMExecContext *ExecContext) {
+	return 0;
+}
+
+uint32 SerializeVectorVMOutputDataSets (FVectorVMSerializeState *SerializeState, struct FVectorVMExecContext *ExecContext) {
+	return 0;
+}
+
+uint32 SerializeVectorVMInputDataSets(FVectorVMSerializeState *SerializeState, TArrayView<FDataSetMeta>, const uint8 * const *ConstantTableData, const int *ConstantTableSizes, int32 ConstantTableCount)
+{
+	return 0;
+}
+
+uint32 SerializeVectorVMOutputDataSets(FVectorVMSerializeState *SerializeState, TArrayView<FDataSetMeta>, const uint8 * const *ConstantTableData, const int *ConstantTableSizes, int32 ConstantTableCount)
+{
+	return 0;
+}
+
+void SerializeVectorVMWriteToFile(FVectorVMSerializeState *SerializeState, uint8 WhichStateWritten, const wchar_t *Filename)
+{
+
+}
+
+void FreeVectorVMSerializeState(FVectorVMSerializeState *SerializeState)
+{
+
+}
+
+#endif //VECTORVM_SUPPORTS_SERIALIZATION
