Experimental Vector VM improvements: faster execution speed, faster optimization speed, uses less runtime memory, added debugging features

diff --git a/Engine/Source/Runtime/VectorVM/Public/VectorVMCommon.h b/Engine/Source/Runtime/VectorVM/Public/VectorVMCommon.h
--- a/Engine/Source/Runtime/VectorVM/Public/VectorVMCommon.h
+++ b/Engine/Source/Runtime/VectorVM/Public/VectorVMCommon.h
@@ -30,203 +30,189 @@ enum class EVectorVMOperandLocation : uint8
 	Num
 };
 
-//              OpCode                       Category   num params	int/float flags
-#define VVM_OP_XM_LIST																																									                                   \
-	VVM_OP_XM( done                         , Other     , 0			, VVM_INS_PARAM_FFFFFF              ) /* 0  */                                                                                                         \
-	VVM_OP_XM( add                          , Op        , 3			, VVM_INS_PARAM_FFFFFF              ) /* 1  */                                                                                                         \
-	VVM_OP_XM( sub                          , Op        , 3			, VVM_INS_PARAM_FFFFFF              ) /* 2  */                                                                                                         \
-	VVM_OP_XM( mul                          , Op        , 3			, VVM_INS_PARAM_FFFFFF              ) /* 3  */                                                                                                         \
-	VVM_OP_XM( div                          , Op        , 3			, VVM_INS_PARAM_FFFFFF              ) /* 4  */                                                                                                         \
-	VVM_OP_XM( mad                          , Op        , 4			, VVM_INS_PARAM_FFFFFF              ) /* 5  */                                                                                                         \
-	VVM_OP_XM( lerp                         , Op        , 4			, VVM_INS_PARAM_FFFFFF              ) /* 6  */                                                                                                         \
-	VVM_OP_XM( rcp                          , Op        , 2			, VVM_INS_PARAM_FFFFFF              ) /* 7  */                                                                                                         \
-	VVM_OP_XM( rsq                          , Op        , 2			, VVM_INS_PARAM_FFFFFF              ) /* 8  */                                                                                                         \
-	VVM_OP_XM( sqrt                         , Op        , 2			, VVM_INS_PARAM_FFFFFF              ) /* 9  */                                                                                                         \
-	VVM_OP_XM( neg                          , Op        , 2			, VVM_INS_PARAM_FFFFFF              ) /* 10 */                                                                                                         \
-	VVM_OP_XM( abs                          , Op        , 2			, VVM_INS_PARAM_FFFFFF              ) /* 11 */                                                                                                         \
-	VVM_OP_XM( exp                          , Op        , 2			, VVM_INS_PARAM_FFFFFF              ) /* 12 */                                                                                                         \
-	VVM_OP_XM( exp2                         , Op        , 2			, VVM_INS_PARAM_FFFFFF              ) /* 13 */                                                                                                         \
-	VVM_OP_XM( log                          , Op        , 2			, VVM_INS_PARAM_FFFFFF              ) /* 14 */                                                                                                         \
-	VVM_OP_XM( log2                         , Op        , 2			, VVM_INS_PARAM_FFFFFF              ) /* 15 */                                                                                                         \
-	VVM_OP_XM( sin                          , Op        , 2			, VVM_INS_PARAM_FFFFFF              ) /* 16 */                                                                                                         \
-	VVM_OP_XM( cos                          , Op        , 2			, VVM_INS_PARAM_FFFFFF              ) /* 17 */                                                                                                         \
-	VVM_OP_XM( tan                          , Op        , 2			, VVM_INS_PARAM_FFFFFF              ) /* 18 */                                                                                                         \
-	VVM_OP_XM( asin                         , Op        , 2			, VVM_INS_PARAM_FFFFFF              ) /* 19 */                                                                                                         \
-	VVM_OP_XM( acos                         , Op        , 2			, VVM_INS_PARAM_FFFFFF              ) /* 20 */                                                                                                         \
-	VVM_OP_XM( atan                         , Op        , 2			, VVM_INS_PARAM_FFFFFF              ) /* 21 */                                                                                                         \
-	VVM_OP_XM( atan2                        , Op        , 3			, VVM_INS_PARAM_FFFFFF              ) /* 22 */                                                                                                         \
-	VVM_OP_XM( ceil                         , Op        , 2			, VVM_INS_PARAM_FFFFFF              ) /* 23 */                                                                                                         \
-	VVM_OP_XM( floor                        , Op        , 2			, VVM_INS_PARAM_FFFFFF              ) /* 24 */                                                                                                         \
-	VVM_OP_XM( fmod                         , Op        , 3			, VVM_INS_PARAM_FFFFFF              ) /* 25 */                                                                                                         \
-	VVM_OP_XM( frac                         , Op        , 2			, VVM_INS_PARAM_FFFFFF              ) /* 26 */                                                                                                         \
-	VVM_OP_XM( trunc                        , Op        , 2			, VVM_INS_PARAM_FFFFFF              ) /* 27 */                                                                                                         \
-	VVM_OP_XM( clamp                        , Op        , 4			, VVM_INS_PARAM_FFFFFF              ) /* 28 */                                                                                                         \
-	VVM_OP_XM( min                          , Op        , 3			, VVM_INS_PARAM_FFFFFF              ) /* 29 */                                                                                                         \
-	VVM_OP_XM( max                          , Op        , 3			, VVM_INS_PARAM_FFFFFF              ) /* 30 */                                                                                                         \
-	VVM_OP_XM( pow                          , Op        , 3			, VVM_INS_PARAM_FFFFFF              ) /* 31 */                                                                                                         \
-	VVM_OP_XM( round                        , Op        , 2			, VVM_INS_PARAM_FFFFFF              ) /* 32 */                                                                                                         \
-	VVM_OP_XM( sign                         , Op        , 2			, VVM_INS_PARAM_FFFFFF              ) /* 33 */                                                                                                         \
-	VVM_OP_XM( step                         , Op        , 3			, VVM_INS_PARAM_FFFFFF              ) /* 34 */                                                                                                         \
-	VVM_OP_XM( random                       , Op        , 2			, VVM_INS_PARAM_FFFFFF              ) /* 35 */                                                                                                         \
-	VVM_OP_XM( noise                        , Op        , 0			, VVM_INS_PARAM_FFFFFF              ) /* 36 */                                                                                                         \
-	VVM_OP_XM( cmplt                        , Op        , 3			, VVM_INS_PARAM_FFFIFF              ) /* 37 */                                                                                                         \
-	VVM_OP_XM( cmple                        , Op        , 3			, VVM_INS_PARAM_FFFIFF              ) /* 38 */                                                                                                         \
-	VVM_OP_XM( cmpgt                        , Op        , 3			, VVM_INS_PARAM_FFFIFF              ) /* 39 */                                                                                                         \
-	VVM_OP_XM( cmpge                        , Op        , 3			, VVM_INS_PARAM_FFFIFF              ) /* 40 */                                                                                                         \
-	VVM_OP_XM( cmpeq                        , Op        , 3			, VVM_INS_PARAM_FFFIFF              ) /* 41 */                                                                                                         \
-	VVM_OP_XM( cmpneq                       , Op        , 3			, VVM_INS_PARAM_FFFIFF              ) /* 42 */                                                                                                         \
-	VVM_OP_XM( select                       , Op        , 4			, VVM_INS_PARAM_FFIFFF              ) /* 43 */                                                                                                         \
-	VVM_OP_XM( addi                         , Op        , 3			, VVM_INS_PARAM_FFIIII              ) /* 44 */                                                                                                         \
-	VVM_OP_XM( subi                         , Op        , 3			, VVM_INS_PARAM_FFIIII              ) /* 45 */                                                                                                         \
-	VVM_OP_XM( muli                         , Op        , 3			, VVM_INS_PARAM_FFIIII              ) /* 46 */                                                                                                         \
-	VVM_OP_XM( divi                         , Op        , 3			, VVM_INS_PARAM_FFIIII              ) /* 47 //SSE Integer division is not implemented as an intrinsic. Will have to do some manual implementation. */  \
-	VVM_OP_XM( clampi                       , Op        , 4			, VVM_INS_PARAM_FFIIII              ) /* 48 */                                                                                                         \
-	VVM_OP_XM( mini                         , Op        , 3			, VVM_INS_PARAM_FFIIII              ) /* 49 */                                                                                                         \
-	VVM_OP_XM( maxi                         , Op        , 3			, VVM_INS_PARAM_FFIIII              ) /* 50 */                                                                                                         \
-	VVM_OP_XM( absi                         , Op        , 2			, VVM_INS_PARAM_FFIIII              ) /* 51 */                                                                                                         \
-	VVM_OP_XM( negi                         , Op        , 2			, VVM_INS_PARAM_FFIIII              ) /* 52 */                                                                                                         \
-	VVM_OP_XM( signi                        , Op        , 2			, VVM_INS_PARAM_FFIIII              ) /* 53 */                                                                                                         \
-	VVM_OP_XM( randomi                      , Op        , 2			, VVM_INS_PARAM_FFIIII              ) /* 54 */                                                                                                         \
-	VVM_OP_XM( cmplti                       , Op        , 3			, VVM_INS_PARAM_FFIIII              ) /* 55 */                                                                                                         \
-	VVM_OP_XM( cmplei                       , Op        , 3			, VVM_INS_PARAM_FFIIII              ) /* 56 */                                                                                                         \
-	VVM_OP_XM( cmpgti                       , Op        , 3			, VVM_INS_PARAM_FFIIII              ) /* 57 */                                                                                                         \
-	VVM_OP_XM( cmpgei                       , Op        , 3			, VVM_INS_PARAM_FFIIII              ) /* 58 */                                                                                                         \
-	VVM_OP_XM( cmpeqi                       , Op        , 3			, VVM_INS_PARAM_FFIIII              ) /* 59 */                                                                                                         \
-	VVM_OP_XM( cmpneqi                      , Op        , 3			, VVM_INS_PARAM_FFIIII              ) /* 60 */                                                                                                         \
-	VVM_OP_XM( bit_and                      , Op        , 3			, VVM_INS_PARAM_FFIIII              ) /* 61 */                                                                                                         \
-	VVM_OP_XM( bit_or                       , Op        , 3			, VVM_INS_PARAM_FFIIII              ) /* 62 */                                                                                                         \
-	VVM_OP_XM( bit_xor                      , Op        , 3			, VVM_INS_PARAM_FFIIII              ) /* 63 */                                                                                                         \
-	VVM_OP_XM( bit_not                      , Op        , 3			, VVM_INS_PARAM_FFIIII              ) /* 64 */                                                                                                         \
-	VVM_OP_XM( bit_lshift                   , Op        , 3			, VVM_INS_PARAM_FFIIII              ) /* 65 */                                                                                                         \
-	VVM_OP_XM( bit_rshift                   , Op        , 3			, VVM_INS_PARAM_FFIIII              ) /* 66 */                                                                                                         \
-	VVM_OP_XM( logic_and                    , Op        , 3			, VVM_INS_PARAM_FFIIII              ) /* 67 */                                                                                                         \
-	VVM_OP_XM( logic_or                     , Op        , 3			, VVM_INS_PARAM_FFIIII              ) /* 68 */                                                                                                         \
-	VVM_OP_XM( logic_xor                    , Op        , 3			, VVM_INS_PARAM_FFIIII              ) /* 69 */                                                                                                         \
-	VVM_OP_XM( logic_not                    , Op        , 2			, VVM_INS_PARAM_FFIIII              ) /* 70 */                                                                                                         \
-	VVM_OP_XM( f2i                          , Op        , 2			, VVM_INS_PARAM_FFFFFI              ) /* 71 */                                                                                                         \
-	VVM_OP_XM( i2f                          , Op        , 2			, VVM_INS_PARAM_FFFFIF              ) /* 72 */                                                                                                         \
-	VVM_OP_XM( f2b                          , Op        , 2			, VVM_INS_PARAM_FFFFFI              ) /* 73 */                                                                                                         \
-	VVM_OP_XM( b2f                          , Op        , 2			, VVM_INS_PARAM_FFFFIF              ) /* 74 */                                                                                                         \
-	VVM_OP_XM( i2b                          , Op        , 2			, VVM_INS_PARAM_FFFFII              ) /* 75 */                                                                                                         \
-	VVM_OP_XM( b2i                          , Op        , 2			, VVM_INS_PARAM_FFFFII              ) /* 76 */                                                                                                         \
-	VVM_OP_XM( inputdata_float              , Input     , 3			, VVM_INS_PARAM_FFFFFF              ) /* 77 */                                                                                                         \
-	VVM_OP_XM( inputdata_int32              , Input     , 3			, VVM_INS_PARAM_FFFFFF              ) /* 78 */                                                                                                         \
-	VVM_OP_XM( inputdata_half               , Input     , 3			, VVM_INS_PARAM_FFFFFF              ) /* 79 */                                                                                                         \
-	VVM_OP_XM( inputdata_noadvance_float    , Input     , 3			, VVM_INS_PARAM_FFFFFF              ) /* 80 */                                                                                                         \
-	VVM_OP_XM( inputdata_noadvance_int32    , Input     , 3			, VVM_INS_PARAM_FFFFFF              ) /* 81 */                                                                                                         \
-	VVM_OP_XM( inputdata_noadvance_half     , Input     , 3			, VVM_INS_PARAM_FFFFFF              ) /* 82 */                                                                                                         \
-	VVM_OP_XM( outputdata_float             , Output    , 3			, VVM_INS_PARAM_FFFFFI              ) /* 83 */                                                                                                         \
-	VVM_OP_XM( outputdata_int32             , Output    , 3			, VVM_INS_PARAM_FFFIII              ) /* 84 */                                                                                                         \
-	VVM_OP_XM( outputdata_half              , Output    , 3			, VVM_INS_PARAM_FFFFFI              ) /* 85 */                                                                                                         \
-	VVM_OP_XM( acquireindex                 , IndexGen  , 2			, VVM_INS_PARAM_FFFFII              ) /* 86 */                                                                                                         \
-	VVM_OP_XM( external_func_call           , ExtFnCall , 0			, VVM_INS_PARAM_FFFFFF              ) /* 87 */                                                                                                         \
-	VVM_OP_XM( exec_index                   , Op        , 1         , VVM_INS_PARAM_FFFFFI              ) /* 88 */                                                                                                         \
-	VVM_OP_XM( noise2D                      , Other     , 0			, VVM_INS_PARAM_FFFFFF              ) /* 89 */                                                                                                         \
-	VVM_OP_XM( noise3D                      , Other     , 0			, VVM_INS_PARAM_FFFFFF              ) /* 90 */                                                                                                         \
-	VVM_OP_XM( enter_stat_scope             , Stat      , 1			, VVM_INS_PARAM_FFFFFF              ) /* 91 */                                                                                                         \
-	VVM_OP_XM( exit_stat_scope              , Stat      , 0			, VVM_INS_PARAM_FFFFFF              ) /* 92 */                                                                                                         \
-	VVM_OP_XM( update_id                    , RWBuffer  , 3			, VVM_INS_PARAM_FFFFFF              ) /* 93 */                                                                                                         \
-	VVM_OP_XM( acquire_id                   , RWBuffer  , 3			, VVM_INS_PARAM_FFFFFF              ) /* 94 */                                                                                                         \
-	VVM_OP_XM( half_to_float                , Fused     , 2			, VVM_INS_PARAM_FFFFFF              ) /* 95 */                                                                                                         \
-	VVM_OP_XM( fasi                         , Op        , 2			, VVM_INS_PARAM_FFFFFF              ) /* 96 */                                                                                                         \
-	VVM_OP_XM( iasf                         , Op        , 2			, VVM_INS_PARAM_FFFFFF              ) /* 97 */                                                                                                         \
-	/* Fused ops -- combined ops that show up frequently together in Fornite.  There are three types: */                                                                                                                   \
-	/* 1. exec_index that get immediately fed into an add or i2f */                                                                                                                                                        \
-	/* 2. ops with identical inputs */                                                                                                                                                                                     \
-	/* 3. ops where the output of one chain to the next.. ie a mul that feeds directly into a sub */                                                                                                                       \
-	VVM_OP_XM( exec_indexf                  , Op        , 2         , VVM_INS_PARAM_FFIIIF              ) /* 98 */                                                                                                         \
-	VVM_OP_XM( exec_index_addi              , Op        , 2         , VVM_INS_PARAM_FFIIII              ) /* 90 */                                                                                                         \
-	\
-	VVM_OP_XM( cmplt_select                 , Op        , 5         , VVM_INS_PARAM_FFFFFF              ) /* 100 */                                                                                                         \
-	VVM_OP_XM( cmple_select                 , Op        , 5         , VVM_INS_PARAM_FFFFFF              ) /* 101 */                                                                                                        \
-	VVM_OP_XM( cmpeq_select                 , Op        , 5         , VVM_INS_PARAM_FFFFFF              ) /* 102 */                                                                                                        \
-	VVM_OP_XM( cmplti_select                , Op        , 5         , VVM_INS_PARAM_FIIFFF              ) /* 103 */                                                                                                        \
-	VVM_OP_XM( cmplei_select                , Op        , 5         , VVM_INS_PARAM_FIIFFF              ) /* 104 */                                                                                                        \
-	VVM_OP_XM( cmpeqi_select                , Op        , 5         , VVM_INS_PARAM_FIIFFF              ) /* 105 */                                                                                                        \
-	\
-	VVM_OP_XM( cmplt_logic_and              , Op        , 4         , VVM_INS_PARAM_FFFFII              ) /* 106 */                                                                                                        \
-	VVM_OP_XM( cmple_logic_and              , Op        , 4         , VVM_INS_PARAM_FFFFII              ) /* 107 */                                                                                                        \
-	VVM_OP_XM( cmpgt_logic_and              , Op        , 4         , VVM_INS_PARAM_FFFFII              ) /* 108 */                                                                                                        \
-	VVM_OP_XM( cmpge_logic_and              , Op        , 4         , VVM_INS_PARAM_FFFFII              ) /* 109 */                                                                                                        \
-	VVM_OP_XM( cmpeq_logic_and              , Op        , 4         , VVM_INS_PARAM_FFFFII              ) /* 110 */                                                                                                        \
-	VVM_OP_XM( cmpne_logic_and              , Op        , 4         , VVM_INS_PARAM_FFFFII              ) /* 111 */                                                                                                        \
-	VVM_OP_XM( cmplti_logic_and             , Op        , 4         , VVM_INS_PARAM_FIIFII              ) /* 112 */                                                                                                        \
-	VVM_OP_XM( cmplei_logic_and             , Op        , 4         , VVM_INS_PARAM_FIIFII              ) /* 113 */                                                                                                        \
-	VVM_OP_XM( cmpgti_logic_and             , Op        , 4         , VVM_INS_PARAM_FIIFII              ) /* 114 */                                                                                                        \
-	VVM_OP_XM( cmpgei_logic_and             , Op        , 4         , VVM_INS_PARAM_FIIFIF              ) /* 115 */                                                                                                        \
-	VVM_OP_XM( cmpeqi_logic_and             , Op        , 4         , VVM_INS_PARAM_FIIFIF              ) /* 116 */                                                                                                        \
-	VVM_OP_XM( cmpnei_logic_and             , Op        , 4         , VVM_INS_PARAM_FIIFIF              ) /* 117 */                                                                                                        \
-	\
-	VVM_OP_XM( cmplt_logic_or               , Op        , 4         , VVM_INS_PARAM_FFFFFF              ) /* 118 */                                                                                                        \
-	VVM_OP_XM( cmple_logic_or               , Op        , 4         , VVM_INS_PARAM_FFFFFF              ) /* 119 */                                                                                                        \
-	VVM_OP_XM( cmpgt_logic_or               , Op        , 4         , VVM_INS_PARAM_FFFFFF              ) /* 120 */                                                                                                        \
-	VVM_OP_XM( cmpge_logic_or               , Op        , 4         , VVM_INS_PARAM_FFFFFF              ) /* 121 */                                                                                                        \
-	VVM_OP_XM( cmpeq_logic_or               , Op        , 4         , VVM_INS_PARAM_FFFFFF              ) /* 122 */                                                                                                        \
-	VVM_OP_XM( cmpne_logic_or               , Op        , 4         , VVM_INS_PARAM_FFFFFF              ) /* 123 */                                                                                                        \
-	VVM_OP_XM( cmplti_logic_or              , Op        , 4         , VVM_INS_PARAM_FIIFFF              ) /* 124 */                                                                                                        \
-	VVM_OP_XM( cmplei_logic_or              , Op        , 4         , VVM_INS_PARAM_FIIFFF              ) /* 125 */                                                                                                        \
-	VVM_OP_XM( cmpgti_logic_or              , Op        , 4         , VVM_INS_PARAM_FIIFFF              ) /* 126 */                                                                                                        \
-	VVM_OP_XM( cmpgei_logic_or              , Op        , 4         , VVM_INS_PARAM_FIIFFF              ) /* 127 */                                                                                                        \
-	VVM_OP_XM( cmpeqi_logic_or              , Op        , 4         , VVM_INS_PARAM_FIIFFF              ) /* 128 */                                                                                                        \
-	VVM_OP_XM( cmpnei_logic_or              , Op        , 4         , VVM_INS_PARAM_FIIFFF              ) /* 129 */                                                                                                        \
-	\
-	VVM_OP_XM( mad_add                      , Op        , 5         , VVM_INS_PARAM_FFFFFF              ) /* 130 */                                                                                                        \
-	VVM_OP_XM( mad_sub0                     , Op        , 5         , VVM_INS_PARAM_FFFFFF              ) /* 131 */                                                                                                        \
-	VVM_OP_XM( mad_sub1                     , Op        , 5         , VVM_INS_PARAM_FFFFFF              ) /* 132 */                                                                                                        \
-	VVM_OP_XM( mad_mul                      , Op        , 5         , VVM_INS_PARAM_FFFFFF              ) /* 133 */                                                                                                        \
-	VVM_OP_XM( mad_sqrt                     , Op        , 4         , VVM_INS_PARAM_FFFFFF              ) /* 134 */                                                                                                        \
-	VVM_OP_XM( mad_mad0                     , Op        , 6         , VVM_INS_PARAM_FFFFFF              ) /* 135 */                                                                                                        \
-	VVM_OP_XM( mad_mad1                     , Op        , 6         , VVM_INS_PARAM_FFFFFF              ) /* 136 */                                                                                                        \
-	\
-	VVM_OP_XM( mul_mad0                     , Op        , 5         , VVM_INS_PARAM_FFFFFF              ) /* 137 */                                                                                                        \
-	VVM_OP_XM( mul_mad1                     , Op        , 5         , VVM_INS_PARAM_FFFFFF              ) /* 138 */                                                                                                        \
-	VVM_OP_XM( mul_add                      , Op        , 4         , VVM_INS_PARAM_FFFFFF              ) /* 139 */                                                                                                        \
-	VVM_OP_XM( mul_sub0                     , Op        , 4         , VVM_INS_PARAM_FFFFFF              ) /* 140 */                                                                                                        \
-	VVM_OP_XM( mul_sub1                     , Op        , 4         , VVM_INS_PARAM_FFFFFF              ) /* 141 */                                                                                                        \
-	VVM_OP_XM( mul_mul                      , Op        , 4         , VVM_INS_PARAM_FFFFFF              ) /* 142 */                                                                                                        \
-	VVM_OP_XM( mul_max                      , Op        , 4         , VVM_INS_PARAM_FFFFFF              ) /* 143 */                                                                                                        \
-	VVM_OP_XM( mul_2x                       , Op        , 4         , VVM_INS_PARAM_FFFFFF              ) /* 144 */                                                                                                        \
-	\
-	VVM_OP_XM( add_mad1                     , Op        , 5         , VVM_INS_PARAM_FFFFFF              ) /* 145 */                                                                                                        \
-	VVM_OP_XM( add_add                      , Op        , 4         , VVM_INS_PARAM_FFFFFF              ) /* 146 */                                                                                                        \
-	\
-	VVM_OP_XM( sub_cmplt1                   , Op        , 4         , VVM_INS_PARAM_FFFFFI              ) /* 147 */                                                                                                        \
-	VVM_OP_XM( sub_neg                      , Op        , 3         , VVM_INS_PARAM_FFFFFF              ) /* 148 */                                                                                                        \
-	VVM_OP_XM( sub_mul                      , Op        , 4         , VVM_INS_PARAM_FFFFFF              ) /* 149 */                                                                                                        \
-	\
-	VVM_OP_XM( div_mad0                     , Op        , 5         , VVM_INS_PARAM_FFFFFF              ) /* 150 */                                                                                                        \
-	VVM_OP_XM( div_f2i                      , Op        , 3         , VVM_INS_PARAM_FFFFFI              ) /* 151 */                                                                                                        \
-	VVM_OP_XM( div_mul                      , Op        , 4         , VVM_INS_PARAM_FFFFFF              ) /* 152 */                                                                                                        \
-	\
-	VVM_OP_XM( muli_addi                    , Op        , 4         , VVM_INS_PARAM_IIIIII              ) /* 153 */                                                                                                        \
-	VVM_OP_XM( addi_bit_rshift				, Op        , 4         , VVM_INS_PARAM_IIIIII              ) /* 154 */                                                                                                        \
-	VVM_OP_XM( addi_muli                    , Op        , 4         , VVM_INS_PARAM_IIIIII              ) /* 155 */                                                                                                        \
-	\
-	VVM_OP_XM( b2i_2x                       , Op        , 3         , VVM_INS_PARAM_IIIIII              ) /* 156 */                                                                                                        \
-	\
-	VVM_OP_XM( i2f_div0                     , Op        , 3         , VVM_INS_PARAM_FFFFFF              ) /* 157 */                                                                                                        \
-	VVM_OP_XM( i2f_div1                     , Op        , 3         , VVM_INS_PARAM_FFFFFF              ) /* 158 */                                                                                                        \
-	VVM_OP_XM( i2f_mul                      , Op        , 3         , VVM_INS_PARAM_FFFFFF              ) /* 159 */                                                                                                        \
-	VVM_OP_XM( i2f_mad0                     , Op        , 4         , VVM_INS_PARAM_FFFFFF              ) /* 160 */                                                                                                        \
-	VVM_OP_XM( i2f_mad1                     , Op        , 4         , VVM_INS_PARAM_FFFFFF              ) /* 161 */                                                                                                        \
-	\
-	VVM_OP_XM( f2i_select1                  , Op        , 4         , VVM_INS_PARAM_IIFIII              ) /* 162 */                                                                                                        \
-	VVM_OP_XM( f2i_maxi                     , Op        , 3         , VVM_INS_PARAM_IIIFII              ) /* 163 */                                                                                                        \
-	VVM_OP_XM( f2i_addi                     , Op        , 3         , VVM_INS_PARAM_IIIFII              ) /* 164 */                                                                                                        \
-	\
-	VVM_OP_XM( fmod_add                     , Op        , 3         , VVM_INS_PARAM_FFFFFF              ) /* 165 */                                                                                                        \
-	VVM_OP_XM( bit_and_i2f                  , Op        , 3         , VVM_INS_PARAM_FFFFFF              ) /* 166 */                                                                                                        \
-	VVM_OP_XM( bit_rshift_bit_and           , Op        , 3         , VVM_INS_PARAM_FFFFFF              ) /* 167 */                                                                                                        \
-	VVM_OP_XM( neg_cmplt                    , Op        , 3         , VVM_INS_PARAM_FFFFFF              ) /* 168 */                                                                                                        \
-	VVM_OP_XM( bit_or_muli                  , Op        , 3         , VVM_INS_PARAM_FFFFFF              ) /* 169 */                                                                                                        \
-	VVM_OP_XM( bit_lshift_bit_or            , Op        , 3         , VVM_INS_PARAM_FFFFFF              ) /* 170 */                                                                                                        \
-	VVM_OP_XM( random_add                   , Op        , 3         , VVM_INS_PARAM_FFFFFF              ) /* 171 */                                                                                                        \
-	VVM_OP_XM( random_2x                    , Op        , 3         , VVM_INS_PARAM_FFFFFF              ) /* 172 */                                                                                                        \
-	VVM_OP_XM( max_f2i                      , Op        , 3         , VVM_INS_PARAM_FFFFFF              ) /* 173 */                                                                                                        \
-	VVM_OP_XM( select_mul                   , Op        , 5         , VVM_INS_PARAM_FFFFFF              ) /* 174 */                                                                                                        \
-	VVM_OP_XM( select_add                   , Op        , 5         , VVM_INS_PARAM_FFFFFF              ) /* 175 */                                                                                                        \
-	VVM_OP_XM( sin_cos                      , Op        , 3         , VVM_INS_PARAM_FFFFFF              ) /* 176 */
-
+//            OpCode                         Category    #in    #out  dispatch,  merge tbl offset, merge tbl count, int/float flags
+#define VVM_OP_XM_LIST                                                                                                                                       \
+	VVM_OP_XM( done                         , Other     , 0    , 0   , done     , 0             , 0               , VVM_INS_PARAM_FFFFFF) /* 0  */           \
+	VVM_OP_XM( add                          , Op        , 2    , 1   , f        , 0             , 2               , VVM_INS_PARAM_FFFFFF) /* 1  */           \
+	VVM_OP_XM( sub                          , Op        , 2    , 1   , f        , 2             , 3               , VVM_INS_PARAM_FFFFFF) /* 2  */           \
+	VVM_OP_XM( mul                          , Op        , 2    , 1   , f        , 5             , 8               , VVM_INS_PARAM_FFFFFF) /* 3  */           \
+	VVM_OP_XM( div                          , Op        , 2    , 1   , f        , 13            , 3               , VVM_INS_PARAM_FFFFFF) /* 4  */           \
+	VVM_OP_XM( mad                          , Op        , 3    , 1   , f        , 16            , 7               , VVM_INS_PARAM_FFFFFF) /* 5  */           \
+	VVM_OP_XM( lerp                         , Op        , 3    , 1   , f        , 0             , 0               , VVM_INS_PARAM_FFFFFF) /* 6  */           \
+	VVM_OP_XM( rcp                          , Op        , 1    , 1   , f        , 0             , 0               , VVM_INS_PARAM_FFFFFF) /* 7  */           \
+	VVM_OP_XM( rsq                          , Op        , 1    , 1   , f        , 0             , 0               , VVM_INS_PARAM_FFFFFF) /* 8  */           \
+	VVM_OP_XM( sqrt                         , Op        , 1    , 1   , f        , 0             , 0               , VVM_INS_PARAM_FFFFFF) /* 9  */           \
+	VVM_OP_XM( neg                          , Op        , 1    , 1   , f        , 23            , 1               , VVM_INS_PARAM_FFFFFF) /* 10 */           \
+	VVM_OP_XM( abs                          , Op        , 1    , 1   , f        , 0             , 0               , VVM_INS_PARAM_FFFFFF) /* 11 */           \
+	VVM_OP_XM( exp                          , Op        , 1    , 1   , f        , 0             , 0               , VVM_INS_PARAM_FFFFFF) /* 12 */           \
+	VVM_OP_XM( exp2                         , Op        , 1    , 1   , f        , 0             , 0               , VVM_INS_PARAM_FFFFFF) /* 13 */           \
+	VVM_OP_XM( log                          , Op        , 1    , 1   , f        , 0             , 0               , VVM_INS_PARAM_FFFFFF) /* 14 */           \
+	VVM_OP_XM( log2                         , Op        , 1    , 1   , f        , 0             , 0               , VVM_INS_PARAM_FFFFFF) /* 15 */           \
+	VVM_OP_XM( sin                          , Op        , 1    , 1   , f        , 24            , 1               , VVM_INS_PARAM_FFFFFF) /* 16 */           \
+	VVM_OP_XM( cos                          , Op        , 1    , 1   , f        , 25            , 1               , VVM_INS_PARAM_FFFFFF) /* 17 */           \
+	VVM_OP_XM( tan                          , Op        , 1    , 1   , f        , 0             , 0               , VVM_INS_PARAM_FFFFFF) /* 18 */           \
+	VVM_OP_XM( asin                         , Op        , 1    , 1   , f        , 0             , 0               , VVM_INS_PARAM_FFFFFF) /* 19 */           \
+	VVM_OP_XM( acos                         , Op        , 1    , 1   , f        , 0             , 0               , VVM_INS_PARAM_FFFFFF) /* 20 */           \
+	VVM_OP_XM( atan                         , Op        , 1    , 1   , f        , 0             , 0               , VVM_INS_PARAM_FFFFFF) /* 21 */           \
+	VVM_OP_XM( atan2                        , Op        , 2    , 1   , f        , 0             , 0               , VVM_INS_PARAM_FFFFFF) /* 22 */           \
+	VVM_OP_XM( ceil                         , Op        , 1    , 1   , f        , 0             , 0               , VVM_INS_PARAM_FFFFFF) /* 23 */           \
+	VVM_OP_XM( floor                        , Op        , 1    , 1   , f        , 0             , 0               , VVM_INS_PARAM_FFFFFF) /* 24 */           \
+	VVM_OP_XM( fmod                         , Op        , 2    , 1   , f        , 26            , 1               , VVM_INS_PARAM_FFFFFF) /* 25 */           \
+	VVM_OP_XM( frac                         , Op        , 1    , 1   , f        , 0             , 0               , VVM_INS_PARAM_FFFFFF) /* 26 */           \
+	VVM_OP_XM( trunc                        , Op        , 1    , 1   , f        , 0             , 0               , VVM_INS_PARAM_FFFFFF) /* 27 */           \
+	VVM_OP_XM( clamp                        , Op        , 3    , 1   , f        , 0             , 0               , VVM_INS_PARAM_FFFFFF) /* 28 */           \
+	VVM_OP_XM( min                          , Op        , 2    , 1   , f        , 0             , 0               , VVM_INS_PARAM_FFFFFF) /* 29 */           \
+	VVM_OP_XM( max                          , Op        , 2    , 1   , f        , 27            , 1               , VVM_INS_PARAM_FFFFFF) /* 30 */           \
+	VVM_OP_XM( pow                          , Op        , 2    , 1   , f        , 0             , 0               , VVM_INS_PARAM_FFFFFF) /* 31 */           \
+	VVM_OP_XM( round                        , Op        , 1    , 1   , f        , 0             , 0               , VVM_INS_PARAM_FFFFFF) /* 32 */           \
+	VVM_OP_XM( sign                         , Op        , 1    , 1   , f        , 0             , 0               , VVM_INS_PARAM_FFFFFF) /* 33 */           \
+	VVM_OP_XM( step                         , Op        , 2    , 1   , f        , 0             , 0               , VVM_INS_PARAM_FFFFFF) /* 34 */           \
+	VVM_OP_XM( random                       , Op        , 1    , 1   , null     , 28            , 2               , VVM_INS_PARAM_FFFFFF) /* 35 */           \
+	VVM_OP_XM( noise                        , Op        , 0    , 0   , null     , 0             , 0               , VVM_INS_PARAM_FFFFFF) /* 36 */           \
+	VVM_OP_XM( cmplt                        , Op        , 2    , 1   , f        , 30            , 3               , VVM_INS_PARAM_FFFIFF) /* 37 */           \
+	VVM_OP_XM( cmple                        , Op        , 2    , 1   , f        , 33            , 3               , VVM_INS_PARAM_FFFIFF) /* 38 */           \
+	VVM_OP_XM( cmpgt                        , Op        , 2    , 1   , f        , 36            , 2               , VVM_INS_PARAM_FFFIFF) /* 39 */           \
+	VVM_OP_XM( cmpge                        , Op        , 2    , 1   , f        , 38            , 2               , VVM_INS_PARAM_FFFIFF) /* 40 */           \
+	VVM_OP_XM( cmpeq                        , Op        , 2    , 1   , f        , 40            , 3               , VVM_INS_PARAM_FFFIFF) /* 41 */           \
+	VVM_OP_XM( cmpneq                       , Op        , 2    , 1   , f        , 43            , 2               , VVM_INS_PARAM_FFFIFF) /* 42 */           \
+	VVM_OP_XM( select                       , Op        , 3    , 1   , f        , 45            , 2               , VVM_INS_PARAM_FFFFFI) /* 43 */           \
+	VVM_OP_XM( addi                         , Op        , 2    , 1   , i        , 47            , 2               , VVM_INS_PARAM_FFIIII) /* 44 */           \
+	VVM_OP_XM( subi                         , Op        , 2    , 1   , i        , 0             , 0               , VVM_INS_PARAM_FFIIII) /* 45 */           \
+	VVM_OP_XM( muli                         , Op        , 2    , 1   , i        , 49            , 1               , VVM_INS_PARAM_FFIIII) /* 46 */           \
+	VVM_OP_XM( divi                         , Op        , 2    , 1   , i        , 0             , 0               , VVM_INS_PARAM_FFIIII) /* 47 */           \
+	VVM_OP_XM( clampi                       , Op        , 3    , 1   , i        , 0             , 0               , VVM_INS_PARAM_FFIIII) /* 48 */           \
+	VVM_OP_XM( mini                         , Op        , 2    , 1   , i        , 0             , 0               , VVM_INS_PARAM_FFIIII) /* 49 */           \
+	VVM_OP_XM( maxi                         , Op        , 2    , 1   , i        , 0             , 0               , VVM_INS_PARAM_FFIIII) /* 50 */           \
+	VVM_OP_XM( absi                         , Op        , 1    , 1   , i        , 0             , 0               , VVM_INS_PARAM_FFIIII) /* 51 */           \
+	VVM_OP_XM( negi                         , Op        , 1    , 1   , i        , 0             , 0               , VVM_INS_PARAM_FFIIII) /* 52 */           \
+	VVM_OP_XM( signi                        , Op        , 1    , 1   , i        , 0             , 0               , VVM_INS_PARAM_FFIIII) /* 53 */           \
+	VVM_OP_XM( randomi                      , Op        , 1    , 1   , null     , 0             , 0               , VVM_INS_PARAM_FFIIII) /* 54 */           \
+	VVM_OP_XM( cmplti                       , Op        , 2    , 1   , i        , 50            , 3               , VVM_INS_PARAM_FFIIII) /* 55 */           \
+	VVM_OP_XM( cmplei                       , Op        , 2    , 1   , i        , 53            , 3               , VVM_INS_PARAM_FFIIII) /* 56 */           \
+	VVM_OP_XM( cmpgti                       , Op        , 2    , 1   , i        , 59            , 2               , VVM_INS_PARAM_FFIIII) /* 57 */           \
+	VVM_OP_XM( cmpgei                       , Op        , 2    , 1   , i        , 61            , 2               , VVM_INS_PARAM_FFIIII) /* 58 */           \
+	VVM_OP_XM( cmpeqi                       , Op        , 2    , 1   , i        , 56            , 3               , VVM_INS_PARAM_FFIIII) /* 59 */           \
+	VVM_OP_XM( cmpneqi                      , Op        , 2    , 1   , i        , 63            , 2               , VVM_INS_PARAM_FFIIII) /* 60 */           \
+	VVM_OP_XM( bit_and                      , Op        , 2    , 1   , i        , 65            , 1               , VVM_INS_PARAM_FFIIII) /* 61 */           \
+	VVM_OP_XM( bit_or                       , Op        , 2    , 1   , i        , 66            , 1               , VVM_INS_PARAM_FFIIII) /* 62 */           \
+	VVM_OP_XM( bit_xor                      , Op        , 2    , 1   , i        , 0             , 0               , VVM_INS_PARAM_FFIIII) /* 63 */           \
+	VVM_OP_XM( bit_not                      , Op        , 1    , 1   , i        , 0             , 0               , VVM_INS_PARAM_FFIIII) /* 64 */           \
+	VVM_OP_XM( bit_lshift                   , Op        , 2    , 1   , i        , 67            , 1               , VVM_INS_PARAM_FFIIII) /* 65 */           \
+	VVM_OP_XM( bit_rshift                   , Op        , 2    , 1   , i        , 68            , 1               , VVM_INS_PARAM_FFIIII) /* 66 */           \
+	VVM_OP_XM( logic_and                    , Op        , 2    , 1   , i        , 0             , 0               , VVM_INS_PARAM_FFIIII) /* 67 */           \
+	VVM_OP_XM( logic_or                     , Op        , 2    , 1   , i        , 0             , 0               , VVM_INS_PARAM_FFIIII) /* 68 */           \
+	VVM_OP_XM( logic_xor                    , Op        , 2    , 1   , i        , 0             , 0               , VVM_INS_PARAM_FFIIII) /* 69 */           \
+	VVM_OP_XM( logic_not                    , Op        , 1    , 1   , i        , 0             , 0               , VVM_INS_PARAM_FFIIII) /* 70 */           \
+	VVM_OP_XM( f2i                          , Op        , 1    , 1   , i        , 74            , 3               , VVM_INS_PARAM_FFFFFI) /* 71 */           \
+	VVM_OP_XM( i2f                          , Op        , 1    , 1   , f        , 69            , 5               , VVM_INS_PARAM_FFFFIF) /* 72 */           \
+	VVM_OP_XM( f2b                          , Op        , 1    , 1   , f        , 0             , 0               , VVM_INS_PARAM_FFFFFI) /* 73 */           \
+	VVM_OP_XM( b2f                          , Op        , 1    , 1   , f        , 0             , 0               , VVM_INS_PARAM_FFFFIF) /* 74 */           \
+	VVM_OP_XM( i2b                          , Op        , 1    , 1   , i        , 0             , 0               , VVM_INS_PARAM_FFFFII) /* 75 */           \
+	VVM_OP_XM( b2i                          , Op        , 1    , 1   , i        , 77            , 1               , VVM_INS_PARAM_FFFFII) /* 76 */           \
+	VVM_OP_XM( inputdata_float              , Input     , 0    , 0   , null     , 0             , 0               , VVM_INS_PARAM_FFFFFF) /* 77 */           \
+	VVM_OP_XM( inputdata_int32              , Input     , 0    , 0   , null     , 0             , 0               , VVM_INS_PARAM_FFFFFF) /* 78 */           \
+	VVM_OP_XM( inputdata_half               , Input     , 0    , 0   , null     , 0             , 0               , VVM_INS_PARAM_FFFFFF) /* 79 */           \
+	VVM_OP_XM( inputdata_noadvance_float    , Input     , 0    , 0   , null     , 0             , 0               , VVM_INS_PARAM_FFFFFF) /* 80 */           \
+	VVM_OP_XM( inputdata_noadvance_int32    , Input     , 0    , 0   , null     , 0             , 0               , VVM_INS_PARAM_FFFFFF) /* 81 */           \
+	VVM_OP_XM( inputdata_noadvance_half     , Input     , 0    , 0   , null     , 0             , 0               , VVM_INS_PARAM_FFFFFF) /* 82 */           \
+	VVM_OP_XM( outputdata_float             , Output    , 0    , 0   , null     , 0             , 0               , VVM_INS_PARAM_FFFFFI) /* 83 */           \
+	VVM_OP_XM( outputdata_int32             , Output    , 0    , 0   , null     , 0             , 0               , VVM_INS_PARAM_FFFIII) /* 84 */           \
+	VVM_OP_XM( outputdata_half              , Output    , 0    , 0   , null     , 0             , 0               , VVM_INS_PARAM_FFFFFI) /* 85 */           \
+	VVM_OP_XM( acquireindex                 , IndexGen  , 1    , 0   , null     , 0             , 0               , VVM_INS_PARAM_FFFFII) /* 86 */           \
+	VVM_OP_XM( external_func_call           , ExtFnCall , 0    , 0   , null     , 0             , 0               , VVM_INS_PARAM_FFFFFF) /* 87 */           \
+	VVM_OP_XM( exec_index                   , Op        , 1    , 1   , null     , 78            , 2               , VVM_INS_PARAM_FFFFFI) /* 88 */           \
+	VVM_OP_XM( noise2D                      , Other     , 0    , 0   , null     , 0             , 0               , VVM_INS_PARAM_FFFFFF) /* 89 */           \
+	VVM_OP_XM( noise3D                      , Other     , 0    , 0   , null     , 0             , 0               , VVM_INS_PARAM_FFFFFF) /* 90 */           \
+	VVM_OP_XM( enter_stat_scope             , Stat      , 1    , 0   , null     , 0             , 0               , VVM_INS_PARAM_FFFFFF) /* 91 */           \
+	VVM_OP_XM( exit_stat_scope              , Stat      , 0    , 0   , null     , 0             , 0               , VVM_INS_PARAM_FFFFFF) /* 92 */           \
+	VVM_OP_XM( update_id                    , RWBuffer  , 0    , 0   , null     , 0             , 0               , VVM_INS_PARAM_FFFFFF) /* 93 */           \
+	VVM_OP_XM( acquire_id                   , RWBuffer  , 1    , 0   , null     , 0             , 0               , VVM_INS_PARAM_FFFFFF) /* 94 */           \
+	VVM_OP_XM( half_to_float                , Op        , 1    , 1   , null     , 0             , 0               , VVM_INS_PARAM_FFFFFF) /* 95 */           \
+	VVM_OP_XM( fasi                         , Op        , 1    , 1   , null     , 0             , 0               , VVM_INS_PARAM_FFFFFF) /* 96 */           \
+	VVM_OP_XM( iasf                         , Op        , 1    , 1   , null     , 0             , 0               , VVM_INS_PARAM_FFFFFI) /* 97 */           \
+	/* Merged ops -- combined ops that show up frequently together in Fornite.  There are three types: */   \
+	/* 1. exec_index that get immediately fed into an add or i2f                                       */   \
+	/* 2. ops with identical inputs                                                                    */   \
+	/* 3. ops where the output of one chain to the next.. ie a mul that feeds directly into a sub      */   \
+	VVM_OP_XM( exec_indexf                  , Op        , 1    , 1   , null   , 0    , 0   , VVM_INS_PARAM_FFIIIF) /* 96 */            \
+	VVM_OP_XM( exec_index_addi              , Op        , 1    , 1   , null   , 0    , 0   , VVM_INS_PARAM_FFIIII) /* 97 */            \
+	VVM_OP_XM( cmplt_select                 , Op        , 4    , 1   , f      , 0    , 0   , VVM_INS_PARAM_FFFFFF) /* 98 */            \
+	VVM_OP_XM( cmple_select                 , Op        , 4    , 1   , f      , 0    , 0   , VVM_INS_PARAM_FFFFFF) /* 99 */            \
+	VVM_OP_XM( cmpeq_select                 , Op        , 4    , 1   , f      , 0    , 0   , VVM_INS_PARAM_FFFFFF) /* 100 */           \
+	VVM_OP_XM( cmplti_select                , Op        , 4    , 1   , f      , 0    , 0   , VVM_INS_PARAM_FIIFFF) /* 101 */           \
+	VVM_OP_XM( cmplei_select                , Op        , 4    , 1   , f      , 0    , 0   , VVM_INS_PARAM_FIIFFF) /* 102 */           \
+	VVM_OP_XM( cmpeqi_select                , Op        , 4    , 1   , f      , 0    , 0   , VVM_INS_PARAM_FIIFFF) /* 103 */           \
+	VVM_OP_XM( cmplt_logic_and              , Op        , 3    , 1   , f      , 0    , 0   , VVM_INS_PARAM_FFFFII) /* 104 */           \
+	VVM_OP_XM( cmple_logic_and              , Op        , 3    , 1   , f      , 0    , 0   , VVM_INS_PARAM_FFFFII) /* 105 */           \
+	VVM_OP_XM( cmpgt_logic_and              , Op        , 3    , 1   , f      , 0    , 0   , VVM_INS_PARAM_FFFFII) /* 106 */           \
+	VVM_OP_XM( cmpge_logic_and              , Op        , 3    , 1   , f      , 0    , 0   , VVM_INS_PARAM_FFFFII) /* 107 */           \
+	VVM_OP_XM( cmpeq_logic_and              , Op        , 3    , 1   , f      , 0    , 0   , VVM_INS_PARAM_FFFFII) /* 108 */           \
+	VVM_OP_XM( cmpne_logic_and              , Op        , 3    , 1   , f      , 0    , 0   , VVM_INS_PARAM_FFFFII) /* 109 */           \
+	VVM_OP_XM( cmplti_logic_and             , Op        , 3    , 1   , i      , 0    , 0   , VVM_INS_PARAM_FIIFII) /* 110 */           \
+	VVM_OP_XM( cmplei_logic_and             , Op        , 3    , 1   , i      , 0    , 0   , VVM_INS_PARAM_FIIFII) /* 111 */           \
+	VVM_OP_XM( cmpgti_logic_and             , Op        , 3    , 1   , i      , 0    , 0   , VVM_INS_PARAM_FIIFII) /* 112 */           \
+	VVM_OP_XM( cmpgei_logic_and             , Op        , 3    , 1   , i      , 0    , 0   , VVM_INS_PARAM_FIIFIF) /* 113 */           \
+	VVM_OP_XM( cmpeqi_logic_and             , Op        , 3    , 1   , i      , 0    , 0   , VVM_INS_PARAM_FIIFIF) /* 114 */           \
+	VVM_OP_XM( cmpnei_logic_and             , Op        , 3    , 1   , i      , 0    , 0   , VVM_INS_PARAM_FIIFIF) /* 115 */           \
+	VVM_OP_XM( cmplt_logic_or               , Op        , 3    , 1   , f      , 0    , 0   , VVM_INS_PARAM_FFFFFF) /* 116 */           \
+	VVM_OP_XM( cmple_logic_or               , Op        , 3    , 1   , f      , 0    , 0   , VVM_INS_PARAM_FFFFFF) /* 117 */           \
+	VVM_OP_XM( cmpgt_logic_or               , Op        , 3    , 1   , f      , 0    , 0   , VVM_INS_PARAM_FFFFFF) /* 118 */           \
+	VVM_OP_XM( cmpge_logic_or               , Op        , 3    , 1   , f      , 0    , 0   , VVM_INS_PARAM_FFFFFF) /* 119 */           \
+	VVM_OP_XM( cmpeq_logic_or               , Op        , 3    , 1   , f      , 0    , 0   , VVM_INS_PARAM_FFFFFF) /* 120 */           \
+	VVM_OP_XM( cmpne_logic_or               , Op        , 3    , 1   , f      , 0    , 0   , VVM_INS_PARAM_FFFFFF) /* 121 */           \
+	VVM_OP_XM( cmplti_logic_or              , Op        , 3    , 1   , i      , 0    , 0   , VVM_INS_PARAM_FIIFFF) /* 122 */           \
+	VVM_OP_XM( cmplei_logic_or              , Op        , 3    , 1   , i      , 0    , 0   , VVM_INS_PARAM_FIIFFF) /* 123 */           \
+	VVM_OP_XM( cmpgti_logic_or              , Op        , 3    , 1   , i      , 0    , 0   , VVM_INS_PARAM_FIIFFF) /* 124 */           \
+	VVM_OP_XM( cmpgei_logic_or              , Op        , 3    , 1   , i      , 0    , 0   , VVM_INS_PARAM_FIIFFF) /* 125 */           \
+	VVM_OP_XM( cmpeqi_logic_or              , Op        , 3    , 1   , i      , 0    , 0   , VVM_INS_PARAM_FIIFFF) /* 126 */           \
+	VVM_OP_XM( cmpnei_logic_or              , Op        , 3    , 1   , i      , 0    , 0   , VVM_INS_PARAM_FIIFFF) /* 127 */           \
+	VVM_OP_XM( mad_add                      , Op        , 4    , 1   , f      , 0    , 0   , VVM_INS_PARAM_FFFFFF) /* 128 */           \
+	VVM_OP_XM( mad_sub0                     , Op        , 4    , 1   , f      , 0    , 0   , VVM_INS_PARAM_FFFFFF) /* 129 */           \
+	VVM_OP_XM( mad_sub1                     , Op        , 4    , 1   , f      , 0    , 0   , VVM_INS_PARAM_FFFFFF) /* 130 */           \
+	VVM_OP_XM( mad_mul                      , Op        , 4    , 1   , f      , 0    , 0   , VVM_INS_PARAM_FFFFFF) /* 131 */           \
+	VVM_OP_XM( mad_sqrt                     , Op        , 3    , 1   , f      , 0    , 0   , VVM_INS_PARAM_FFFFFF) /* 132 */           \
+	VVM_OP_XM( mad_mad0                     , Op        , 5    , 1   , f      , 0    , 0   , VVM_INS_PARAM_FFFFFF) /* 133 */           \
+	VVM_OP_XM( mad_mad1                     , Op        , 5    , 1   , f      , 0    , 0   , VVM_INS_PARAM_FFFFFF) /* 134 */           \
+	VVM_OP_XM( mul_mad0                     , Op        , 4    , 1   , f      , 0    , 0   , VVM_INS_PARAM_FFFFFF) /* 135 */           \
+	VVM_OP_XM( mul_mad1                     , Op        , 4    , 1   , f      , 0    , 0   , VVM_INS_PARAM_FFFFFF) /* 136 */           \
+	VVM_OP_XM( mul_add                      , Op        , 3    , 1   , f      , 0    , 0   , VVM_INS_PARAM_FFFFFF) /* 137 */           \
+	VVM_OP_XM( mul_sub0                     , Op        , 3    , 1   , f      , 0    , 0   , VVM_INS_PARAM_FFFFFF) /* 138 */           \
+	VVM_OP_XM( mul_sub1                     , Op        , 3    , 1   , f      , 0    , 0   , VVM_INS_PARAM_FFFFFF) /* 139 */           \
+	VVM_OP_XM( mul_mul                      , Op        , 3    , 1   , f      , 0    , 0   , VVM_INS_PARAM_FFFFFF) /* 140 */           \
+	VVM_OP_XM( mul_max                      , Op        , 3    , 1   , f      , 0    , 0   , VVM_INS_PARAM_FFFFFF) /* 141 */           \
+	VVM_OP_XM( mul_2x                       , Op        , 2    , 2   , f      , 0    , 0   , VVM_INS_PARAM_FFFFFF) /* 142 */           \
+	VVM_OP_XM( add_mad1                     , Op        , 4    , 1   , f      , 0    , 0   , VVM_INS_PARAM_FFFFFF) /* 143 */           \
+	VVM_OP_XM( add_add                      , Op        , 3    , 1   , f      , 0    , 0   , VVM_INS_PARAM_FFFFFF) /* 144 */           \
+	VVM_OP_XM( sub_cmplt1                   , Op        , 3    , 1   , f      , 0    , 0   , VVM_INS_PARAM_FFIFFF) /* 145 */           \
+	VVM_OP_XM( sub_neg                      , Op        , 2    , 1   , f      , 0    , 0   , VVM_INS_PARAM_FFFFFF) /* 146 */           \
+	VVM_OP_XM( sub_mul                      , Op        , 3    , 1   , f      , 0    , 0   , VVM_INS_PARAM_FFFFFF) /* 147 */           \
+	VVM_OP_XM( div_mad0                     , Op        , 4    , 1   , f      , 0    , 0   , VVM_INS_PARAM_FFFFFF) /* 148 */           \
+	VVM_OP_XM( div_f2i                      , Op        , 2    , 1   , i      , 0    , 0   , VVM_INS_PARAM_FFFFFI) /* 149 */           \
+	VVM_OP_XM( div_mul                      , Op        , 3    , 1   , f      , 0    , 0   , VVM_INS_PARAM_FFFFFF) /* 150 */           \
+	VVM_OP_XM( muli_addi                    , Op        , 3    , 1   , i      , 0    , 0   , VVM_INS_PARAM_IIIIII) /* 151 */           \
+	VVM_OP_XM( addi_bit_rshift				, Op        , 3    , 1   , i      , 0    , 0   , VVM_INS_PARAM_IIIIII) /* 152 */           \
+	VVM_OP_XM( addi_muli                    , Op        , 3    , 1   , i      , 0    , 0   , VVM_INS_PARAM_IIIIII) /* 153 */           \
+	VVM_OP_XM( b2i_2x                       , Op        , 1    , 2   , i      , 0    , 0   , VVM_INS_PARAM_IIIIII) /* 154 */           \
+	VVM_OP_XM( i2f_div0                     , Op        , 2    , 1   , f      , 0    , 0   , VVM_INS_PARAM_FFFFFF) /* 155 */           \
+	VVM_OP_XM( i2f_div1                     , Op        , 2    , 1   , f      , 0    , 0   , VVM_INS_PARAM_FFFFFF) /* 156 */           \
+	VVM_OP_XM( i2f_mul                      , Op        , 2    , 1   , f      , 0    , 0   , VVM_INS_PARAM_FFFFFF) /* 157 */           \
+	VVM_OP_XM( i2f_mad0                     , Op        , 3    , 1   , f      , 0    , 0   , VVM_INS_PARAM_FFFFFF) /* 158 */           \
+	VVM_OP_XM( i2f_mad1                     , Op        , 3    , 1   , f      , 0    , 0   , VVM_INS_PARAM_FFFFFF) /* 159 */           \
+	VVM_OP_XM( f2i_select1                  , Op        , 3    , 1   , i      , 0    , 0   , VVM_INS_PARAM_IIFIII) /* 160 */           \
+	VVM_OP_XM( f2i_maxi                     , Op        , 2    , 1   , i      , 0    , 0   , VVM_INS_PARAM_IIIFII) /* 161 */           \
+	VVM_OP_XM( f2i_addi                     , Op        , 2    , 1   , i      , 0    , 0   , VVM_INS_PARAM_IIIFII) /* 162 */           \
+	VVM_OP_XM( fmod_add                     , Op        , 3    , 1   , f      , 0    , 0   , VVM_INS_PARAM_FFFFFF) /* 163 */           \
+	VVM_OP_XM( bit_and_i2f                  , Op        , 2    , 1   , f      , 0    , 0   , VVM_INS_PARAM_FFFFFF) /* 164 */           \
+	VVM_OP_XM( bit_rshift_bit_and           , Op        , 3    , 1   , i      , 0    , 0   , VVM_INS_PARAM_FFFFFF) /* 165 */           \
+	VVM_OP_XM( neg_cmplt                    , Op        , 2    , 1   , f      , 0    , 0   , VVM_INS_PARAM_FFFFFF) /* 166 */           \
+	VVM_OP_XM( bit_or_muli                  , Op        , 3    , 1   , i      , 0    , 0   , VVM_INS_PARAM_FFFFFF) /* 167 */           \
+	VVM_OP_XM( bit_lshift_bit_or            , Op        , 3    , 1   , i      , 0    , 0   , VVM_INS_PARAM_FFFFFF) /* 168 */           \
+	VVM_OP_XM( random_add                   , Op        , 2    , 1   , null   , 0    , 0   , VVM_INS_PARAM_FFFFFF) /* 169 */           \
+	VVM_OP_XM( random_2x                    , Op        , 1    , 2   , null   , 0    , 0   , VVM_INS_PARAM_FFFFFF) /* 170 */           \
+	VVM_OP_XM( max_f2i                      , Op        , 2    , 1   , i      , 0    , 0   , VVM_INS_PARAM_FFFFFF) /* 171 */           \
+	VVM_OP_XM( select_mul                   , Op        , 4    , 1   , f      , 0    , 0   , VVM_INS_PARAM_FFFFFF) /* 172 */           \
+	VVM_OP_XM( select_add                   , Op        , 4    , 1   , f      , 0    , 0   , VVM_INS_PARAM_FFFFFF) /* 173 */           \
+	VVM_OP_XM( sin_cos                      , Op        , 1    , 2   , null   , 0    , 0   , VVM_INS_PARAM_FFFFFF) /* 174 */           
 #define VVM_OP_XM(n, ...) n,
 enum class EVectorVMOp : uint8 {
 	VVM_OP_XM_LIST
@@ -358,6 +344,9 @@ private:
 	FDataSetMeta& operator=(const FDataSetMeta&) = delete;
 };
 
+typedef void * (VectorVMReallocFn)              (void *Ptr, size_t NumBytes, const char *Filename, int LineNumber);
+typedef void   (VectorVMFreeFn)                 (void *Ptr, const char *Filename, int LineNumber);
+
 #if VECTORVM_SUPPORTS_EXPERIMENTAL && VECTORVM_SUPPORTS_LEGACY
 
 class FVectorVMExternalFunctionContextProxy;
diff --git a/Engine/Source/Runtime/VectorVM/Public/VectorVMExperimental.h b/Engine/Source/Runtime/VectorVM/Public/VectorVMExperimental.h
--- a/Engine/Source/Runtime/VectorVM/Public/VectorVMExperimental.h
+++ b/Engine/Source/Runtime/VectorVM/Public/VectorVMExperimental.h
@@ -1,18 +1,15 @@
 // Copyright Epic Games, Inc. All Rights Reserved.
 #pragma once
 
-// HEADER_UNIT_SKIP - Not included directly
+#include "VectorVMSerialization.h"
+#include "VectorVMCommon.h"
 
-#define VECTORVM_SUPPORTS_AVX 0
-#define VECTORVM_SUPPORTS_COMPUTED_GOTO 0
+#if VECTORVM_SUPPORTS_EXPERIMENTAL
 
-//only to be included by VectorVM.h
+#define VECTORVM_SUPPORTS_COMPUTED_GOTO 0
 
-typedef void * (VectorVMReallocFn)              (void *Ptr, size_t NumBytes, const char *Filename, int LineNumber);
-typedef void   (VectorVMFreeFn)                 (void *Ptr, const char *Filename, int LineNumber);
 struct FDataSetMeta;
 
-
 #define VVM_INS_PARAM_FFFFFF	0b000000
 #define VVM_INS_PARAM_FFFFFI	0b000001
 #define VVM_INS_PARAM_FFFFIF	0b000010
@@ -78,7 +75,6 @@ struct FDataSetMeta;
 #define VVM_INS_PARAM_IIIIIF	0b111110
 #define VVM_INS_PARAM_IIIIII	0b111111
 
-
 #define VVM_OP_CAT_XM_LIST    \
 	VVM_OP_CAT_XM(Input)      \
 	VVM_OP_CAT_XM(Output)     \
@@ -87,7 +83,6 @@ struct FDataSetMeta;
 	VVM_OP_CAT_XM(IndexGen)   \
 	VVM_OP_CAT_XM(RWBuffer)   \
 	VVM_OP_CAT_XM(Stat)       \
-	VVM_OP_CAT_XM(Fused)      \
 	VVM_OP_CAT_XM(Other)
 
 #define VVM_OP_CAT_XM(Category, ...) Category,
@@ -103,6 +98,27 @@ static const EVectorVMOpCategory VVM_OP_CATEGORIES[] = {
 #undef VVM_OP_XM
 };
 
+#define VVM_MIN(a, b)               ((a) < (b) ? (a) : (b))
+#define VVM_MAX(a, b)               ((a) > (b) ? (a) : (b))
+#define VVM_CLAMP(v, min, max)      ((v) < (min) ? (min) : ((v) < (max) ? (v) : (max)))
+#define VVM_ALIGN(num, alignment)   (((size_t)(num) + (alignment) - 1) & ~((alignment) - 1))
+#define VVM_ALIGN_4(num)            (((size_t)(num) + 3) & ~3)
+#define VVM_ALIGN_8(num)            (((size_t)(num) + 7) & ~7)
+#define VVM_ALIGN_16(num)           (((size_t)(num) + 15) & ~15)
+#define VVM_ALIGN_32(num)           (((size_t)(num) + 31) & ~31)
+#define VVM_ALIGN_64(num)           (((size_t)(num) + 63) & ~63)
+
+#define VVM_PTR_ALIGN    VVM_ALIGN_16
+#define VVM_REG_SIZE     sizeof(FVecReg)
+
+enum {
+	VVM_RT_TEMPREG,
+	VVM_RT_CONST,
+	VVM_RT_INPUT,
+	VVM_RT_OUTPUT,
+	VVM_RT_INVALID
+};
+
 enum EVectorVMFlags
 {
 	VVMFlag_OptSaveIntermediateState = 1 << 0,
@@ -110,176 +126,10 @@ enum EVectorVMFlags
 	VVMFlag_LargeScript              = 1 << 2,   //if true register indices are 16 bit, otherwise they're 8 bit
 	VVMFlag_HasRandInstruction       = 1 << 3,
 	VVMFlag_DataMapCacheSetup        = 1 << 4,
-	VVMFlag_SupportsAVX              = 1 << 5,
 };
 
 //prototypes for serialization that are required whether or not serialization is enabled
-#if VECTORVM_SUPPORTS_EXPERIMENTAL || defined(VVM_INCLUDE_SERIALIZATION)
-
-typedef uint32 (VectorVMOptimizeErrorCallback)  (struct FVectorVMOptimizeContext *OptimizeContext, uint32 ErrorFlags);	//return new error flags
-typedef uint32 (VectorVMSerializeErrorCallback) (struct FVectorVMSerializeState *SerializeState, uint32 ErrorFlags);
-
-#endif // VECTORVM_SUPPORTS_EXPERIMENTAL || VVM_INCLUDE_SERIALIZATION
-
-//Serialization
-enum EVectorVMSerializeFlags
-{
-	VVMSer_SyncRandom        = 1 << 0,
-	VVMSer_SyncExtFns        = 1 << 1,
-	VVMSer_OptimizedBytecode = 1 << 2,
-};
-
-#ifdef VVM_INCLUDE_SERIALIZATION
-
-enum EVectorVMSerializeError
-{
-	VVMSerErr_OutOfMemory    = 1 << 0,
-	VVMSerErr_Init           = 1 << 1,
-	VVMSerErr_InputDataSets  = 1 << 2,
-	VVMSerErr_OutputDataSets = 1 << 3,
-	VVMSerErr_Instruction    = 1 << 4,
-	VVMSerErr_ConstData      = 1 << 5,
-
-	VVMSerErr_Fatal          = 1 << 31
-};
-
-struct FVectorVMSerializeInstruction
-{
-	uint32   OpStart;
-	uint32   NumOps;
-	uint64   Dt;
-	uint64   DtDecode;
-	uint32 * TempRegisters;
-	uint8 *  TempRegisterFlags;
-};
-
-struct FVectorVMSerializeExternalData
-{
-	wchar_t * Name;
-	uint16    NameLen;
-	uint16    NumInputs;
-	uint16    NumOutputs;
-};
-
-struct FVectorVMSerializeDataSet
-{
-	uint32 * InputBuffers;
-	uint32 * OutputBuffers;
-
-	uint32  InputOffset[4];	    //float, int, half (half must be 0)
-	uint32  OutputOffset[4];    //float, int, half (half must be 0)
-
-	int32   InputInstanceOffset;
-	int32   InputDataSetAccessIndex;
-	int32   InputIDAcquireTag;
-
-	int32   OutputInstanceOffset;
-	int32   OutputDataSetAccessIndex;
-	int32   OutputIDAcquireTag;
-
-	int32 * InputIDTable;
-	int32 * InputFreeIDTable;
-	int32 * InputSpawnedIDTable;
-
-	int32   InputIDTableNum;
-	int32   InputFreeIDTableNum;
-	int32   InputSpawnedIDTableNum;
-
-	int32   InputNumFreeIDs;
-	int32   InputMaxUsedIDs;
-	int32   InputNumSpawnedIDs;
-
-	int32 * OutputIDTable;
-	int32 * OutputFreeIDTable;
-	int32 * OutputSpawnedIDTable;
-
-	int32   OutputIDTableNum;
-	int32   OutputFreeIDTableNum;
-	int32   OutputSpawnedIDTableNum;
-
-	int32   OutputNumFreeIDs;
-	int32   OutputMaxUsedIDs;
-	int32   OutputNumSpawnedIDs;
-};
-
-struct FVectorVMSerializeChunk
-{
-	uint32 ChunkIdx;
-	uint32 BatchIdx;
-	uint32 ExecIdx;
-	uint32 StartInstance;
-	uint32 NumInstances;
-
-	uint32 StartThreadID;
-	uint32 EndThreadID;
-
-	uint64 StartClock;
-	uint64 EndClock;
-	uint64 InsExecTime;
-};
-
-struct FVectorVMSerializeState
-{
-	uint32                              NumInstances;
-	uint32                              NumTempRegisters;
-	uint32                              NumTempRegFlags;  //max of NumTempRegisters and Num Input Registers in each dataset
-	uint32                              ConstDataCount;
-
-	uint32                              Flags;
-
-	FVectorVMSerializeInstruction *     Instructions;
-	uint32                              NumInstructions;
-	uint32                              NumInstructionsAllocated;
-
-	uint32                              NumExternalData;
-	FVectorVMSerializeExternalData *    ExternalData;
-	uint32								MaxExtFnRegisters;
-	uint32								MaxExtFnUsed;
-
-	uint64                              ExecDt;
-	uint64                              SerializeDt;
-
-	uint8 *                             TempRegFlags;
-	uint8 *                             Bytecode;
-	uint32                              NumBytecodeBytes;
-
-	FVectorVMSerializeDataSet *         DataSets;
-	uint32                              NumDataSets;
-
-	uint32 *                            ConstTableSizesInBytes;
-	uint32 *                            PreExecConstData;
-	uint32 *                            PostExecConstData;
-
-
-	uint32                              NumChunks;
-	FVectorVMSerializeChunk *           Chunks;
-
-	const struct FVectorVMOptimizeContext *   OptimizeCtx;
-
-	volatile int64                      ChunkComplete; //1 bit for each of the first 64 chunks
-
-	VectorVMReallocFn *                  ReallocFn;
-	VectorVMFreeFn *                     FreeFn;
-
-	struct {
-		uint32                           Flags;
-		uint32                           Line;
-		VectorVMSerializeErrorCallback * CallbackFn;
-	} Error;
-};
-#else // VVM_INCLUDE_SERIALIZATION
-
-struct FVectorVMSerializeState
-{
-	uint32              Flags;
-	VectorVMReallocFn * ReallocFn;
-	VectorVMFreeFn *    FreeFn;
-
-};
-
-#endif // VVM_INCLUDE_SERIALIZATION
-
-#if VECTORVM_SUPPORTS_EXPERIMENTAL
+typedef uint32 (VectorVMOptimizeErrorCallback)     (struct FVectorVMOptimizeContext *OptimizeContext, uint32 ErrorFlags);	//return new error flags
 
 union FVecReg {
 	VectorRegister4f v;
@@ -302,41 +152,40 @@ struct FVectorVMOptimizeInstruction
 	uint32              PtrOffsetInOptimizedBytecode;
 	int                 Index; //initial index.  Instructions are moved around and removed and dependency chains are created based on index, so we need to store this.
 	int					InsMergedIdx; //if not -1, then the instruction index that this is merged with.  Instructions with a set InsMergedIdx are not written to the final bytecode
+	int                 OutputMergeIdx[2]; //if not -1 then this instruction writes directly to an output, not a temp register
+	uint16              RegPtrOffset;
+	int                 NumInputRegisters;
+	int                 NumOutputRegisters;
 	union
 	{
 		struct
 		{
-			uint16 DstRegPtrOffset;
 			uint16 DataSetIdx;
 			uint16 InputIdx;
 		} Input;
 		struct
 		{
-			uint16 RegPtrOffset;
 			uint16 DataSetIdx;
 			uint16 DstRegIdx;
+			int    MergeIdx; //if not -1 then this instruction index is merged with an output or an op, if it's -2 then it's already been taken care of
+			uint16 SerialIdx;
 		} Output;
 		struct
 		{
-			uint32 RegPtrOffset;
-			uint16 NumInputs;
-			uint16 NumOutputs;
+
 		} Op;
 		struct
 		{
-			uint32 RegPtrOffset;
 			uint16 DataSetIdx;
 		} IndexGen;
 		struct
 		{
-			uint32 RegPtrOffset;
 			uint16 ExtFnIdx;
 			uint16 NumInputs;
 			uint16 NumOutputs;
 		} ExtFnCall;
 		struct
 		{
-			uint32 RegPtrOffset;
 			uint16 DataSetIdx;
 		} RWBuffer;
 		struct
@@ -357,7 +206,7 @@ enum EVectorVMOptimizeError
 	VVMOptErr_RegisterUsage         = 1 << 3,
 	VVMOptErr_ConstRemap            = 1 << 4,
 	VVMOptErr_Instructions          = 1 << 5,
-	VVMOptErr_InputFuseBuffer       = 1 << 6,
+	VVMOptErr_InputMergeBuffer      = 1 << 6,
 	VVMOptErr_InstructionReOrder    = 1 << 7,
 	VVMOptErr_SSARemap              = 1 << 8,
 	VVMOptErr_OptimizedBytecode     = 1 << 9,
@@ -373,21 +222,29 @@ struct FVectorVMOptimizeContext
 	uint16 *                              ConstRemap[2];
 	uint16 *                              InputRemapTable;
 	uint16 *                              InputDataSetOffsets;
+	uint8 *                               OutputRemapDataSetIdx;
+	uint16 *                              OutputRemapDataType;
+	uint16 *                              OutputRemapDst;
+
 	FVectorVMExtFunctionData *            ExtFnTable;
 
 	uint32                                NumBytecodeBytes;
-	uint32                                NumOutputDataSets;
+	uint32                                MaxOutputDataSet;
 	uint16                                NumConstsAlloced;    //upper bound to alloc
 	uint32                                NumTempRegisters;
 	uint16                                NumConstsRemapped;
 	uint16                                NumInputsRemapped;
+	uint16                                NumOutputsAliasedToTempRegisters; //@TODO: REMOVE THIS
 	uint16                                NumNoAdvanceInputs;
 	uint16                                NumInputDataSets;
+	uint16                                NumOutputsRemapped;
+	uint16                                NumOutputInstructions;
 	uint32                                NumExtFns;
 	uint32                                MaxExtFnRegisters;
 	uint32                                NumDummyRegsReq;     //External function "null" registers
 	int32                                 MaxExtFnUsed;
 	uint32                                Flags;
+	uint64                                HashId;
 
 	struct
 	{
@@ -409,13 +266,13 @@ struct FVectorVMOptimizeContext
 		uint8 *                           RegisterUsageType;
 		uint16 *                          RegisterUsageBuffer;
 		uint16 *                          SSARegisterUsageBuffer;
-		uint32                            NumBytecodeBytes;
+		uint16 *                          ParentInstructionIdx;
 		uint32                            NumInstructions;
+		uint32                            NumInstructionsAlloced;
 		uint32                            NumRegistersUsed;
 	} Intermediate;                                           //these are freed and NULL after optimize() unless SaveIntermediateData is true when calling OptimizeVectorVMScript
 };
 
-//VectorVMState
 enum EVectorVMStateError
 {
 	VVMErr_InitOutOfMemory     = 1 << 0,
@@ -426,45 +283,34 @@ enum EVectorVMStateError
 	VVMErr_Fatal               = 1 << 31
 };
 
-struct FVectorVMExternalFnPerInstanceData
-{
-	class UNiagaraDataInterface *     DataInterface;
-	void *                            UserData;
-	uint16                            NumInputs;
-	uint16                            NumOutputs;
-};
-
 struct FVectorVMBatchState
 {
 	MS_ALIGN(16) FVecReg *    RegisterData GCC_ALIGN(16);
 	struct {
 		uint32 *              StartingOutputIdxPerDataSet;
 		uint32 *              NumOutputPerDataSet;
+		uint8 **              OutputMaskIdx; //these point to the BatchState's OutputMaskIdx
 		struct {
 			uint32 **         RegData;
 			uint8 *           RegInc;
 			FVecReg *         DummyRegs;
 		} ExtFnDecodedReg;
 		int32 *               RandCounters; //used for external functions only.
+
+		int                   ChunkIdx;
+		int                   StartInstanceThisChunk;
+		int                   NumInstancesThisChunk;
 	} ChunkLocalData;
 	
 	uint8 **                  RegPtrTable;    //not aligned because input pointers could be offest from the DataSetInfo.InstanceOffset, so we must assume every op is unaligned
-	uint8 *                   RegIncTable;
+	uint8 *                   RegIncTable;    //0 for const, 1 for temp reg
+	uint8 *                   OutputMaskIdx;
 
-	int                       StartInstance;
-	int                       NumInstances;
-	
 	union {
 		struct {
 			VectorRegister4i          State[5]; //xorwor state for random/randomi instructions.  DIs use RandomStream.
 			VectorRegister4i          Counters;
 		};
-#		if VECTORVM_SUPPORTS_AVX
-		struct {
-			__m256i                   State[5];
-			__m256i                   Counters;
-		} AVX;
-#		endif
 	} RandState;
 	FRandomStream             RandStream;
 };
@@ -478,11 +324,14 @@ struct FVectorVMState
 
 	FVecReg *                           ConstantBuffers;        //the last OptimizeCtx->NumNoAdvanceInputs are no advance inputs that are copied in the table setup
 	FVectorVMExtFunctionData *          ExtFunctionTable;
-	volatile int32 *                    NumOutputPerDataSet;
+	int32 *                             NumOutputPerDataSet;
 
 	uint16 *                            ConstRemapTable;
 	uint16 *                            InputRemapTable;
 	uint16 *                            InputDataSetOffsets;
+	uint8 *                             OutputRemapDataSetIdx;
+	uint16 *                            OutputRemapDataType;
+	uint16 *                            OutputRemapDst;
 
 	uint8 *                             ConstMapCacheIdx;       //these don't get filled out until Exec() is called because they can't be filled out until the state
 	uint16 *                            ConstMapCacheSrc;       //of const and input buffers from Niagara is unknown until exec() is called.
@@ -499,43 +348,46 @@ struct FVectorVMState
 	uint32                              NumConstBuffers;
 	uint32                              NumInputBuffers;
 	uint32                              NumInputDataSets;
-	uint32                              NumOutputDataSets;
+	uint32                              NumOutputsRemapped;
+	uint32                              NumOutputBuffers;
+	uint32                              MaxOutputDataSet;
 	uint32                              NumDummyRegsRequired;
 
 	//batch stuff
 	uint32                              BatchOverheadSize;
 	uint32                              ChunkLocalDataOutputIdxNumBytes;
 	uint32                              ChunkLocalNumOutputNumBytes;
+	uint32                              ChunkLocalOutputMaskIdxNumBytes;
+
+	uint64                              OptimizerHashId;
+	uint32                              TotalNumBytes;
 
 	struct {
 		uint32          NumBytesRequiredPerBatch;
 		uint32          PerBatchRegisterDataBytesRequired;
-		uint32          NumBatches;
 		uint32          MaxChunksPerBatch;
 		uint32          MaxInstancesPerChunk;
 	} ExecCtxCache;
 };
 
-struct FVectorVMExecContext {
+struct FVectorVMExecContext
+{
 	struct
 	{
 		uint32          NumBytesRequiredPerBatch;
 		uint32          PerBatchRegisterDataBytesRequired;
-		uint32          NumBatches;
 		uint32          MaxChunksPerBatch;
 		uint32          MaxInstancesPerChunk;
-		volatile int32  NumInstancesAssignedToBatches;
-		volatile int32  NumInstancesCompleted;
 	} Internal;
 
-	FVectorVMState *                        VVMState;
-	TArrayView<FDataSetMeta>                DataSets;
+	FVectorVMState *                        VVMState;               //created with AllocVectorVMState()
+	TArrayView<FDataSetMeta>                DataSets;					
 	TArrayView<const FVMExternalFunction *> ExtFunctionTable;
-	TArrayView<void*>                       UserPtrTable;
+	TArrayView<void *>                      UserPtrTable;
 	int32                                   NumInstances;
-	const uint8 * const *                   ConstantTableData;
-	const int *                             ConstantTableSizes;
-	int32                                   ConstantTableCount;
+	const uint8 * const *                   ConstantTableData;      //constant tables consist of an array of pointers
+	const int *                             ConstantTableNumBytes;  //an array of sizes in bytes
+	int32                                   ConstantTableCount;     //how many constant tables.  These tables must match the ones used with OptimizeVectorVMScript()
 };
 
 class FVectorVMExternalFunctionContextExperimental
@@ -552,7 +404,7 @@ public:
 	int                      NumLoops;
 	int                      PerInstanceFnInstanceIdx;
 
-	void** UserPtrTable;
+	void **                  UserPtrTable;
 	int                      NumUserPtrs;
 
 	FRandomStream*           RandStream;
@@ -581,27 +433,27 @@ public:
 
 //API FUNCTIONS
 
-//initialization
-VECTORVM_API void InitVectorVM();
-VECTORVM_API void FreeVectorVM();
-
 //normal functions
-VECTORVM_API FVectorVMState * AllocVectorVMState    (FVectorVMOptimizeContext *OptimizeCtx);
-VECTORVM_API void             FreeVectorVMState     (FVectorVMState *VectorVMState);
-VECTORVM_API void             ExecVectorVMState     (FVectorVMExecContext *ExecCtx, FVectorVMSerializeState *SerializeState, FVectorVMSerializeState *CmpSerializeState);
+VECTORVM_API FVectorVMState * AllocVectorVMState                     (FVectorVMOptimizeContext *OptimizeCtx);
+VECTORVM_API void             FreeVectorVMState                      (FVectorVMState *VectorVMState);
+VECTORVM_API void             ExecVectorVMState                      (FVectorVMExecContext *ExecCtx, FVectorVMSerializeState *SerializeState, FVectorVMSerializeState *CmpSerializeState);
 
 //optimize functions
-VECTORVM_API uint32  OptimizeVectorVMScript                (const uint8 *Bytecode, int BytecodeLen, FVectorVMExtFunctionData *ExtFnIOData, int NumExtFns, FVectorVMOptimizeContext *OptContext, uint32 Flags); //OutContext must be zeroed except the Init struct
-VECTORVM_API void    FreeVectorVMOptimizeContext           (FVectorVMOptimizeContext *Context);
-VECTORVM_API void    FreezeVectorVMOptimizeContext         (const FVectorVMOptimizeContext& Context, TArray<uint8>& ContextData);
-VECTORVM_API void    ReinterpretVectorVMOptimizeContextData(TConstArrayView<uint8> ContextData, FVectorVMOptimizeContext& Context);
+VECTORVM_API uint32           OptimizeVectorVMScript                 (const uint8 *Bytecode, int BytecodeLen, FVectorVMExtFunctionData *ExtFnIOData, int NumExtFns, FVectorVMOptimizeContext *OptContext, uint64 HashId, uint32 Flags); //OutContext must be zeroed except the Init struct
+VECTORVM_API void             FreeVectorVMOptimizeContext            (FVectorVMOptimizeContext *Context);
+VECTORVM_API void             FreezeVectorVMOptimizeContext          (const FVectorVMOptimizeContext& Context, TArray<uint8>& ContextData);
+VECTORVM_API void             ReinterpretVectorVMOptimizeContextData (TConstArrayView<uint8> ContextData, FVectorVMOptimizeContext& Context);
 
 //serialize functions
-VECTORVM_API uint32  SerializeVectorVMInputDataSets  (FVectorVMSerializeState *SerializeState, struct FVectorVMExecContext *ExecContext);
-VECTORVM_API uint32  SerializeVectorVMOutputDataSets (FVectorVMSerializeState *SerializeState, struct FVectorVMExecContext *ExecContext);
+VECTORVM_API uint32           SerializeVectorVMInputDataSets         (FVectorVMSerializeState *SerializeState, FVectorVMExecContext *ExecContext);
+VECTORVM_API uint32           SerializeVectorVMOutputDataSets        (FVectorVMSerializeState *SerializeState, FVectorVMExecContext *ExecContext);
+VECTORVM_API void             SerializeVectorVMWriteToFile           (FVectorVMSerializeState *SerializeState, uint8 WhichStateWritten, const wchar_t *Filename);
+VECTORVM_API void             FreeVectorVMSerializeState             (FVectorVMSerializeState *SerializeState);
 
-VECTORVM_API void    SerializeVectorVMWriteToFile    (FVectorVMSerializeState *SerializeState, uint8 WhichStateWritten, const wchar_t *Filename);
-VECTORVM_API void    FreeVectorVMSerializeState      (FVectorVMSerializeState *SerializeState);
+//debug helper functions
+VECTORVM_API uint16           VVMDbgGetRemappedInputIdx              (FVectorVMState *VectorVMState, int InputRegisterIdx, int DataSetIdx, bool IntReg, uint16 *OptOutIdxInRegTable);
+VECTORVM_API uint16           VVMDbgGetRemappedOutputIdx             (FVectorVMState *VectorVMState, int OutputRegisterIdx, int DataSetIdx, bool IntReg, uint16 *OptOutIdxInRegTable);
+VECTORVM_API uint16           VVMDbgGetRemappedConstIdx              (FVectorVMState *VectorVMState, int ConstTableIdx, int ConstBuffIdx, uint16 *OptOutIdxInRegTable);
 
 #else // VECTORVM_SUPPORTS_EXPERIMENTAL
 
@@ -609,9 +461,13 @@ struct FVectorVMState {
 
 };
 
-VECTORVM_API uint32  SerializeVectorVMInputDataSets (FVectorVMSerializeState *SerializeState, TArrayView<FDataSetMeta>, const uint8 * const *ConstantTableData, const int *ConstantTableSizes, int32 ConstantTableCount); //only use when not calling InitVectorVMState()
-VECTORVM_API uint32  SerializeVectorVMOutputDataSets(FVectorVMSerializeState *SerializeState, TArrayView<FDataSetMeta>, const uint8 * const *ConstantTableData, const int *ConstantTableSizes, int32 ConstantTableCount);
-VECTORVM_API void   SerializeVectorVMWriteToFile    (FVectorVMSerializeState *SerializeState, uint8 WhichStateWritten, const wchar_t *Filename);
-VECTORVM_API void   FreeVectorVMSerializeState      (FVectorVMSerializeState *SerializeState);
+VECTORVM_API void    SerializeVectorVMWriteToFile    (FVectorVMSerializeState *SerializeState, uint8 WhichStateWritten, const wchar_t *Filename);
+VECTORVM_API void    FreeVectorVMSerializeState      (FVectorVMSerializeState *SerializeState);
 
 #endif // VECTORVM_SUPPORTS_EXPERIMENTAL
+
+//things needed for when both experimental and legacy are built
+#if VECTORVM_SUPPORTS_LEGACY
+VECTORVM_API uint32  SerializeVectorVMInputDataSets  (FVectorVMSerializeState *SerializeState, TArrayView<FDataSetMeta>, const uint8 * const *ConstantTableData, const int *ConstantTableSizes, int32 ConstantTableCount); //only use when not calling AllocVectorVMState()
+VECTORVM_API uint32  SerializeVectorVMOutputDataSets (FVectorVMSerializeState *SerializeState, TArrayView<FDataSetMeta>, const uint8 * const *ConstantTableData, const int *ConstantTableSizes, int32 ConstantTableCount);
+#endif
\ No newline at end of file
