Experimental Vector VM improvements: faster execution speed, faster optimization speed, uses less runtime memory, added debugging features

diff --git a/Engine/Source/Runtime/VectorVM/Public/VectorVMSerialization.h b/Engine/Source/Runtime/VectorVM/Public/VectorVMSerialization.h
new file mode 100644
--- /dev/null
+++ b/Engine/Source/Runtime/VectorVM/Public/VectorVMSerialization.h
@@ -0,0 +1,298 @@
+// Copyright Epic Games, Inc. All Rights Reserved.
+#pragma once
+
+#include "VectorVMCommon.h"
+
+#if VECTORVM_SUPPORTS_SERIALIZATION
+
+typedef uint32 (VectorVMSerializeErrorCallback)    (struct FVectorVMSerializeState *SerializeState, uint32 ErrorFlags);
+typedef void   (VectorVMSerializeMismatchCallback) (bool IntVal, EVectorVMOp OpCode, uint32 InstructionIndex, int RegisterIndex, int InstanceIndex);
+
+enum EVectorVMSerializeFlags
+{
+	VVMSer_SyncRandom         = 1 << 0,
+	VVMSer_SyncExtFns         = 1 << 1,
+	VVMSer_OptimizedBytecode  = 1 << 2,
+	VVMSer_IncludeOptContext  = 1 << 3,
+	VVMSer_DiffRegsPerIns     = 1 << 4,
+};
+
+enum EVectorVMSerializeInstructionFlags
+{
+	VVMSerIns_Float            = 1 << 0,
+	VVMSerIns_Int              = 1 << 1,
+};
+
+enum EVectorVMSerializeError
+{
+	VVMSerErr_OutOfMemory    = 1 << 0,
+	VVMSerErr_Init           = 1 << 1,
+	VVMSerErr_InputDataSets  = 1 << 2,
+	VVMSerErr_OutputDataSets = 1 << 3,
+	VVMSerErr_Instruction    = 1 << 4,
+	VVMSerErr_ConstData      = 1 << 5,
+
+	VVMSerErr_Fatal          = 1 << 31
+};
+
+struct FVectorVMSerializeInstruction
+{
+	uint32   OpStart;
+	uint32   NumOps;
+	uint64   Dt;
+	uint64   DtDecode;
+	uint32 * RegisterTable;
+	uint8 *  RegisterFlags;
+};
+
+struct FVectorVMSerializeExternalData
+{
+	wchar_t * Name;
+	uint16    NameLen;
+	uint16    NumInputs;
+	uint16    NumOutputs;
+};
+
+struct FVectorVMSerializeDataSet
+{
+	uint32 * InputBuffers;
+	uint32 * OutputBuffers;
+
+	uint32  InputOffset[4];	    //float, int, half, total
+	uint32  OutputOffset[4];    //float, int, half, total
+
+	int32   InputInstanceOffset;
+	int32   InputDataSetAccessIndex;
+	int32   InputIDAcquireTag;
+
+	int32   OutputInstanceOffset;
+	int32   OutputDataSetAccessIndex;
+	int32   OutputIDAcquireTag;
+
+	int32 * InputIDTable;
+	int32 * InputFreeIDTable;
+	int32 * InputSpawnedIDTable;
+
+	int32   InputIDTableNum;
+	int32   InputFreeIDTableNum;
+	int32   InputSpawnedIDTableNum;
+
+	int32   InputNumFreeIDs;
+	int32   InputMaxUsedIDs;
+	int32   InputNumSpawnedIDs;
+
+	int32 * OutputIDTable;
+	int32 * OutputFreeIDTable;
+	int32 * OutputSpawnedIDTable;
+
+	int32   OutputIDTableNum;
+	int32   OutputFreeIDTableNum;
+	int32   OutputSpawnedIDTableNum;
+
+	int32   OutputNumFreeIDs;
+	int32   OutputMaxUsedIDs;
+	int32   OutputNumSpawnedIDs;
+};
+
+struct FVectorVMSerializeChunk
+{
+	uint32 ChunkIdx;
+	uint32 BatchIdx;
+	uint32 ExecIdx;
+	uint32 StartInstance;
+	uint32 NumInstances;
+
+	uint32 StartThreadID;
+	uint32 EndThreadID;
+
+	uint64 StartClock;
+	uint64 EndClock;
+	uint64 InsExecTime;
+};
+
+struct FVectorVMSerializeState
+{
+	uint32                              NumInstances;
+	uint32                              NumTempRegisters;
+	uint32                              NumInputBuffers;
+	uint32                              NumOutputBuffers;
+	uint32                              NumConstBuffers;
+	uint32                              NumRegisterTable;
+
+	uint32                              Flags;
+
+	FVectorVMSerializeInstruction *     Instructions;
+	uint32                              NumInstructions;
+	uint32                              NumInstructionsAllocated;
+
+	uint32                              NumExternalData;
+	FVectorVMSerializeExternalData *    ExternalData;
+	uint32								MaxExtFnRegisters;
+	uint32								MaxExtFnUsed;
+
+	uint64                              OptimizerHashId;
+
+	uint64                              ExecDt;
+	uint64                              SerializeDt;
+
+	uint8 *                             RegisterTableFlags;
+	uint8 *                             Bytecode;
+	uint32                              NumBytecodeBytes;
+
+	FVectorVMSerializeDataSet *         DataSets;
+	uint32                              NumDataSets;
+
+	uint32 *                            ConstTableSizesInBytes;
+	uint32 *                            PreExecConstData;
+	uint32 *                            PostExecConstData;
+
+	uint8 *                             OutputRemapDataSetIdx;
+	uint16 *                            OutputRemapDataType;
+	uint16 *                            OutputRemapDst;
+
+	uint32                              NumChunks;
+	FVectorVMSerializeChunk *           Chunks;
+
+	const struct FVectorVMOptimizeContext *   OptimizeCtx;
+	
+
+	volatile int64                      ChunkComplete; //1 bit for each of the first 64 chunks
+
+	VectorVMReallocFn *                  ReallocFn;
+	VectorVMFreeFn *                     FreeFn;
+
+	VectorVMSerializeMismatchCallback *  MismatchFn;
+
+
+	//used during the execution of a VVMInstance, not to be serialized or relied out outside of the execution
+	struct {
+		const uint8 *                    StartOpPtr;
+		uint64                           StartCycles;
+		FVectorVMSerializeChunk *        Chunk;
+		int	                             NumInstructionsThisChunk;
+		int                              NumChunks;
+	} Running;
+
+	struct {
+		uint32                           Flags;
+		uint32                           Line;
+		VectorVMSerializeErrorCallback * CallbackFn;
+	} Error;
+};
+
+#define VVMSer_batchStart(...)
+#define VVMSer_batchEnd(...)
+
+#define VVMSer_chunkStart(...)
+#define VVMSer_chunkEnd(...)
+
+#define VVMSer_insStart(...)
+#define VVMSer_insEndDecode(...)
+#define VVMSer_insEnd(...)
+
+
+#define VVMSer_batchStartExp(SerializeState, ...)
+#define VVMSer_batchEndExp(SerializeState, ...)
+
+#define VVMIsRegIdxTempReg(Idx) ((Idx) < ExecCtx->VVMState->NumTempRegisters)
+#define VVMIsRegIdxConst(Idx)   ((Idx) >= ExecCtx->VVMState->NumTempRegisters && (Idx) < ExecCtx->VVMState->NumTempRegisters + ExecCtx->VVMState->NumConstsRemapped)
+#define VVMIsRegIdxInput(Idx)   ((Idx) >= ExecCtx->VVMState->NumTempRegisters + ExecCtx->VVMState->NumConstsRemapped)
+
+#if VECTORVM_DEBUG_PRINTF
+#define VVMSer_instruction(Type, NumParams)		for (int vi = 0; vi <= (int)(NumParams); ++vi) {												\
+													int r = (int)((uint16 *)InsPtr)[vi];														\
+													char c = 'R';																				\
+													if (r >= (int)(ExecCtx->VVMState->NumTempRegisters + ExecCtx->VVMState->NumConstBuffers + ExecCtx->VVMState->NumInputBuffers * 2)) {	\
+														c = 'O';																				\
+														r -= ExecCtx->VVMState->NumTempRegisters + ExecCtx->VVMState->NumConstBuffers + ExecCtx->VVMState->NumInputBuffers * 2;			    \
+													} else if (r >= (int)(ExecCtx->VVMState->NumTempRegisters + ExecCtx->VVMState->NumConstBuffers)) {	\
+														c = 'I';																				\
+														r -= ExecCtx->VVMState->NumTempRegisters + ExecCtx->VVMState->NumConstBuffers;			\
+													} else if (r >= (int)ExecCtx->VVMState->NumTempRegisters) {									\
+														c = 'C';																				\
+														r -= ExecCtx->VVMState->NumTempRegisters;												\
+													}																							\
+													printf(", %c%d", c, r);																		\
+													/*for (int j = 0; j < NumInstancesThisChunk; ++j) { printf("%f", ((float *)BatchState->RegPtrTable[((uint16 *)InsPtr)[vi]])[j]); if (j != NumInstancesThisChunk -1) { printf(", "); } }*/ \
+												} \
+												printf("\n");
+#else
+#define VVMSer_instruction(...)
+#endif
+
+#define VVMSer_chunkStartExp(SerializeState, ChunkIdx_, BatchIdx_)                                                     \
+	if (SerializeState)                                                                                                \
+	{                                                                                                                  \
+		SerializeState->Running.NumInstructionsThisChunk = 0;                                                          \
+		SerializeState->Running.Chunk                    = SerializeState->Chunks + SerializeState->Running.NumChunks; \
+		SerializeState->Running.Chunk->ChunkIdx          = SerializeState->Running.NumChunks++;                        \
+		SerializeState->Running.Chunk->BatchIdx          = BatchIdx_;                                                  \
+		SerializeState->Running.Chunk->NumInstances      = BatchState->ChunkLocalData.NumInstancesThisChunk;           \
+		SerializeState->Running.Chunk->StartInstance     = BatchState->ChunkLocalData.StartInstanceThisChunk;          \
+		SerializeState->Running.Chunk->StartThreadID     = FPlatformTLS::GetCurrentThreadId();                         \
+		SerializeState->Running.Chunk->StartClock        = FPlatformTime::Cycles64();                                  \
+	}
+
+#define VVMSer_chunkEndExp(SerializeState, ...)                                                                                              \
+	if (SerializeState && SerializeState->Running.Chunk)                                                                                     \
+	{                                                                                                                                        \
+		FPlatformAtomics::InterlockedOr(&SerializeState->ChunkComplete, (int64)((1ULL << (uint64)SerializeState->Running.Chunk->ChunkIdx))); \
+		SerializeState->Running.Chunk->EndThreadID = FPlatformTLS::GetCurrentThreadId();                                                     \
+		SerializeState->Running.Chunk->EndClock = FPlatformTime::Cycles64();                                                                 \
+	}
+
+#define VVMSer_insStartExp(SerializeState, ...)                          \
+	if (SerializeState) {                                                \
+		SerializeState->Running.StartOpPtr = InsPtr;                     \
+		SerializeState->Running.StartCycles = FPlatformTime::Cycles64(); \
+	}
+
+#define VVMSer_insEndExp(SerializeState, OpStart, NumOps, ...)                                                             \
+	if (SerializeState)                                                                                                    \
+	{                                                                                                                      \
+		uint64 VVMSerEndExecCycles = FPlatformTime::Cycles64();                                                            \
+		VVMSer_serializeInstruction(SerializeState, CmpSerializeState,                                                     \
+            ExecCtx->VVMState, BatchState,                                                                                 \
+			BatchState->ChunkLocalData.StartInstanceThisChunk, BatchState->ChunkLocalData.NumInstancesThisChunk, NumLoops, \
+			OpStart, NumOps, SerializeState->Running.StartCycles, VVMSerEndExecCycles);                                    \
+		uint64 VVMSerEndSerializeCycles = FPlatformTime::Cycles64();                                                       \
+		SerializeState->SerializeDt += VVMSerEndSerializeCycles - VVMSerEndExecCycles;                                     \
+		SerializeState->Running.Chunk->InsExecTime += VVMSerEndExecCycles - SerializeState->Running.StartCycles;           \
+	}
+#define VVMSer_initSerializationState(ExecCtx, SerializeState, OptimizeContext, Flags) VVMSer_initSerializationState_(SerializeState, ExecCtx, SerializeState ? OptimizeContext : NULL, SerializeState ? (uint32)Flags : 0)
+
+
+#else //VECTORVM_SUPPORTS_SERIALIZATION
+
+struct FVectorVMSerializeState
+{
+	uint32              Flags;
+	VectorVMReallocFn * ReallocFn;
+	VectorVMFreeFn *    FreeFn;
+};
+
+#define VVMSer_batchStart(...)
+#define VVMSer_batchEnd(...)
+
+#define VVMSer_chunkStart(...)
+#define VVMSer_chunkEnd(...)
+
+#define VVMSer_insStart(...)
+#define VVMSer_insEndDecode(...)
+#define VVMSer_insEnd(...)
+
+#define VVMSer_batchStartExp(...)
+#define VVMSer_batchEndExp(...)
+
+#define VVMSer_chunkStartExp(...)
+#define VVMSer_chunkEndExp(...)
+
+#define VVMSer_insStartExp(...)
+#define VVMSer_insEndExp(...)
+
+#define VVMSer_instruction(...)
+
+#define VVMSer_initSerializationState(...)
+
+#endif //VECTORVM_SUPPORTS_SERIALIZATION
+
diff --git a/Engine/Source/Runtime/VectorVM/VectorVM.Build.cs b/Engine/Source/Runtime/VectorVM/VectorVM.Build.cs
--- a/Engine/Source/Runtime/VectorVM/VectorVM.Build.cs
+++ b/Engine/Source/Runtime/VectorVM/VectorVM.Build.cs
@@ -17,7 +17,9 @@ public class VectorVM : ModuleRules
             new string[]
             {
                 "VECTORVM_SUPPORTS_EXPERIMENTAL=1",
-                "VECTORVM_SUPPORTS_LEGACY=1"
+                "VECTORVM_SUPPORTS_LEGACY=1",
+                "VECTORVM_SUPPORTS_SERIALIZATION=0",
+                "VECTORVM_DEBUG_PRINTF=0"
             });
     }
 }
