fixing compiler error

diff --git a/Engine/Source/Runtime/VectorVM/Public/VectorVMExperimental.h b/Engine/Source/Runtime/VectorVM/Public/VectorVMExperimental.h
--- a/Engine/Source/Runtime/VectorVM/Public/VectorVMExperimental.h
+++ b/Engine/Source/Runtime/VectorVM/Public/VectorVMExperimental.h
@@ -3,6 +3,7 @@
 
 #include "VectorVMSerialization.h"
 #include "VectorVMCommon.h"
+#include "Math/RandomStream.h"
 
 #if VECTORVM_SUPPORTS_EXPERIMENTAL
 
