compiler fix

diff --git a/Engine/Source/Runtime/VectorVM/Private/VectorVMExperimental.cpp b/Engine/Source/Runtime/VectorVM/Private/VectorVMExperimental.cpp
--- a/Engine/Source/Runtime/VectorVM/Private/VectorVMExperimental.cpp
+++ b/Engine/Source/Runtime/VectorVM/Private/VectorVMExperimental.cpp
@@ -3,6 +3,7 @@
 
 #include "VectorVMExperimental.h"
 #include "VectorVM.h"
+#include "HAL/ConsoleManager.h"
 
 
 //prototypes for internal functions to the VVM
