Win ARM64 compile fix.

#rnx

diff --git a/Engine/Source/Runtime/VectorVM/Private/VectorVMExperimental.cpp b/Engine/Source/Runtime/VectorVM/Private/VectorVMExperimental.cpp
--- a/Engine/Source/Runtime/VectorVM/Private/VectorVMExperimental.cpp
+++ b/Engine/Source/Runtime/VectorVM/Private/VectorVMExperimental.cpp
@@ -399,8 +399,8 @@ static FORCEINLINE void VVM_floatToHalf(void *output, float const *input)
 {
 	float16x4_t out0 = vcvt_f16_f32(vld1q_f32(input + 0));
 	float16x4_t out1 = vcvt_f16_f32(vld1q_f32(input + 4));
-	vst1_f16((__fp16 *)((char *)output + 0), out0);
-	vst1_f16((__fp16 *)((char *)output + 4), out1);
+	vst1_f16((uint16_t*)((char *)output + 0), out0);
+	vst1_f16((uint16_t*)((char *)output + 4), out1);
 }
 
 
