Fix for EVVM scripts not being tagged properly as including randoms if the script only contains merged random functions (random_add or random_2x)

diff --git a/Engine/Source/Runtime/VectorVM/Private/VectorVMExperimentalOptimizer.cpp b/Engine/Source/Runtime/VectorVM/Private/VectorVMExperimentalOptimizer.cpp
--- a/Engine/Source/Runtime/VectorVM/Private/VectorVMExperimentalOptimizer.cpp
+++ b/Engine/Source/Runtime/VectorVM/Private/VectorVMExperimentalOptimizer.cpp
@@ -810,6 +810,16 @@ uint32 OptimizeVectorVMScript(const uint8 *InBytecode, int InBytecodeLen, FVecto
 	const uint8 *OpPtrIn = InBytecode;
 	const uint8 *OpPtrInEnd = InBytecode + InBytecodeLen;
 
+	// todo - consider a flag on the op or to better futureproof merged ops by ensuring that
+	// if a merged op incorporates one of these ops, then it too must be in this set
+	const TArray<EVectorVMOp> RandInstructionSet =
+	{
+		EVectorVMOp::random,
+		EVectorVMOp::randomi,
+		EVectorVMOp::random_add,
+		EVectorVMOp::random_2x
+	};
+
 	//Step 1: Create Intermediate representation of all Instructions
 	while (OpPtrIn < OpPtrInEnd)
 	{
@@ -2956,7 +2966,8 @@ uint32 OptimizeVectorVMScript(const uint8 *InBytecode, int InBytecodeLen, FVecto
 			if (OptimizedBytecode)
 			{
 				Ins->PtrOffsetInOptimizedBytecode = (uint32)NumOptimizedBytesWritten;
-				if (Ins->OpCode == EVectorVMOp::random || Ins->OpCode == EVectorVMOp::randomi) {
+				if (RandInstructionSet.Contains(Ins->OpCode))
+				{
 					OptContext->Flags |= VVMFlag_HasRandInstruction;
 				}
 			}
