VectorVM aligned store fix for Linux

diff --git a/Engine/Source/Runtime/VectorVM/Private/VectorVMExperimental.cpp b/Engine/Source/Runtime/VectorVM/Private/VectorVMExperimental.cpp
--- a/Engine/Source/Runtime/VectorVM/Private/VectorVMExperimental.cpp
+++ b/Engine/Source/Runtime/VectorVM/Private/VectorVMExperimental.cpp
@@ -1035,7 +1035,7 @@ VM_FORCEINLINE const uint8 *VVM_exec_index(const bool CT_MultipleLoops, const ui
 		uint8 *End = P0 + sizeof(FVecReg) * NumLoops;
 		do
 		{
-			VectorIntStoreAligned(Val, P0);
+			VectorIntStore(Val, P0);
 			Val = VectorIntAdd(Val, Four);
 			P0 += sizeof(FVecReg);
 		} while (P0 < End);
@@ -1058,7 +1058,7 @@ VM_FORCEINLINE const uint8 *VVM_exec_indexf(const bool CT_MultipleLoops, const u
 		do
 		{
 			//this is faster than doing the work natively in floating point
-			VectorStoreAligned(VectorIntToFloat(Val), (float *)P0);
+			VectorStore(VectorIntToFloat(Val), (float *)P0);
 			Val = VectorIntAdd(Val, Four);
 			P0 += sizeof(FVecReg);
 		} while (P0 < End);
@@ -1083,7 +1083,7 @@ VM_FORCEINLINE const uint8 *VVM_exec_index_addi(const bool CT_MultipleLoops, con
 		do {
 			VectorRegister4i R0 = VectorIntLoad(P0);
 			P0 += Inc0;
-			VectorIntStoreAligned(VectorIntAdd(Val, R0), P1);
+			VectorIntStore(VectorIntAdd(Val, R0), P1);
 			Val = VectorIntAdd(Val, Four);
 			P1 += sizeof(FVecReg);
 		} while (P1 < End);
@@ -1217,7 +1217,7 @@ static VM_FORCEINLINE const uint8 *VVM_random(const bool CT_MultipleLoops, const
 				VectorRegister4f R0 = VectorLoad((float *)P0);
 				P0 += Inc0;
 				VectorRegister4f Res = VVM_nextRandom(BatchState, R0);
-				VectorStoreAligned(Res, (float *)P1);
+				VectorStore(Res, (float *)P1);
 				P1 += sizeof(FVecReg);
 			} while (P1 < End);
 		}
@@ -1225,7 +1225,7 @@ static VM_FORCEINLINE const uint8 *VVM_random(const bool CT_MultipleLoops, const
 		{
 			VectorRegister4f R0 = VectorLoad((float *)P0);
 			VectorRegister4f Res = VVM_nextRandom(BatchState, R0);
-			VectorStoreAligned(Res, (float *)P1);
+			VectorStore(Res, (float *)P1);
 		}
 	}
 	return InsPtr + 5;
@@ -1248,7 +1248,7 @@ static VM_FORCEINLINE const uint8 *VVM_randomi(const bool CT_MultipleLoops, cons
 				VectorRegister4f R0 = VectorLoad((float *)P0);
 				P0 += Inc0;
 				VectorRegister4f Res = VVM_nextRandom(BatchState, R0);
-				VectorIntStoreAligned(VectorFloatToInt(Res), P0);
+				VectorIntStore(VectorFloatToInt(Res), P0);
 				P1 += sizeof(FVecReg);
 			} while (P1 < End);
 		}
@@ -1256,7 +1256,7 @@ static VM_FORCEINLINE const uint8 *VVM_randomi(const bool CT_MultipleLoops, cons
 		{
 			VectorRegister4f R0 = VectorLoad((float *)P0);
 			VectorRegister4f Res = VVM_nextRandom(BatchState, R0);
-			VectorIntStoreAligned(VectorFloatToInt(Res), P0);
+			VectorIntStore(VectorFloatToInt(Res), P0);
 		}
 }
 	return InsPtr + 5;
@@ -1284,7 +1284,7 @@ static VM_FORCEINLINE const uint8 *VVM_random_add(const bool CT_MultipleLoops, c
 				P0 += Inc0;
 				P1 += Inc1;
 				VectorRegister4f Res = VectorAdd(VVM_nextRandom(BatchState, R0), R1);
-				VectorStoreAligned(Res, (float *)P2);
+				VectorStore(Res, (float *)P2);
 				P2 += sizeof(FVecReg);
 			} while (P2 < End);
 		}
@@ -1293,7 +1293,7 @@ static VM_FORCEINLINE const uint8 *VVM_random_add(const bool CT_MultipleLoops, c
 			VectorRegister4f R0 = VectorLoad((float *)P0);
 			VectorRegister4f R1 = VectorLoad((float *)P1);
 			VectorRegister4f Res = VectorAdd(VVM_nextRandom(BatchState, R0), R1);
-			VectorStoreAligned(Res, (float *)P2);
+			VectorStore(Res, (float *)P2);
 		}
 	}
 	return InsPtr + 7;
@@ -1318,8 +1318,8 @@ static VM_FORCEINLINE const uint8 *VVM_random_2x(const bool CT_MultipleLoops, co
 				P0 += Inc0;
 				VectorRegister4f Res0 = VVM_nextRandom(BatchState, R0);
 				VectorRegister4f Res1 = VVM_nextRandom(BatchState, R0);
-				VectorStoreAligned(Res0, (float *)P1);
-				VectorStoreAligned(Res1, (float *)P2);
+				VectorStore(Res0, (float *)P1);
+				VectorStore(Res1, (float *)P2);
 				P1 += sizeof(FVecReg);
 				P2 += sizeof(FVecReg);
 			} while (P1 < End);
@@ -1329,8 +1329,8 @@ static VM_FORCEINLINE const uint8 *VVM_random_2x(const bool CT_MultipleLoops, co
 			VectorRegister4f R0 = VectorLoad((float *)P0);
 			VectorRegister4f Res0 = VVM_nextRandom(BatchState, R0);
 			VectorRegister4f Res1 = VVM_nextRandom(BatchState, R0);
-			VectorStoreAligned(Res0, (float *)P1);
-			VectorStoreAligned(Res1, (float *)P2);
+			VectorStore(Res0, (float *)P1);
+			VectorStore(Res1, (float *)P2);
 		}
 	}
 	return InsPtr + 7;
@@ -1931,7 +1931,7 @@ VM_FORCEINLINE const uint8 *VVM_Dispatch_execFn1f_1f(const bool CT_MultipleLoops
 			VectorRegister4f R0 = VectorLoad((float *)P0);
 			P0 += Inc0;
 			VectorRegister4f Res = fn(BatchState, R0);
-			VectorStoreAligned(Res, (float *)P1);
+			VectorStore(Res, (float *)P1);
 			P1 += sizeof(FVecReg);
 		} while (P1 < End);
 	}
@@ -1939,7 +1939,7 @@ VM_FORCEINLINE const uint8 *VVM_Dispatch_execFn1f_1f(const bool CT_MultipleLoops
 	{
 		VectorRegister4f R0 = VectorLoad((float *)P0);
 		VectorRegister4f Res = fn(BatchState, R0);
-		VectorStoreAligned(Res, (float *)P1);
+		VectorStore(Res, (float *)P1);
 	}
 	return InsPtr + 5;
 }
@@ -1962,7 +1962,7 @@ VM_FORCEINLINE const uint8 *VVM_Dispatch_execFn2f_1f(const bool CT_MultipleLoops
 			P0 += Inc0;
 			P1 += Inc1;
 			VectorRegister4f Res = fn(BatchState, R0, R1);
-			VectorStoreAligned(Res, (float *)P2);
+			VectorStore(Res, (float *)P2);
 			P2 += sizeof(FVecReg);
 		} while (P2 < End);
 	}
@@ -1971,7 +1971,7 @@ VM_FORCEINLINE const uint8 *VVM_Dispatch_execFn2f_1f(const bool CT_MultipleLoops
 		VectorRegister4f R0 = VectorLoad((float *)P0);
 		VectorRegister4f R1 = VectorLoad((float *)P1);
 		VectorRegister4f Res = fn(BatchState, R0, R1);
-		VectorStoreAligned(Res, (float *)P2);
+		VectorStore(Res, (float *)P2);
 	}
 	return InsPtr + 7;
 }
@@ -1998,7 +1998,7 @@ VM_FORCEINLINE const uint8 *VVM_Dispatch_execFn3f_1f(const bool CT_MultipleLoops
 			P1 += Inc1;
 			P2 += Inc2;
 			VectorRegister4f Res = fn(BatchState, R0, R1, R2);
-			VectorStoreAligned(Res, (float *)P3);
+			VectorStore(Res, (float *)P3);
 			P3 += sizeof(FVecReg);
 		} while (P3 < End);
 	}
@@ -2008,7 +2008,7 @@ VM_FORCEINLINE const uint8 *VVM_Dispatch_execFn3f_1f(const bool CT_MultipleLoops
 		VectorRegister4f R1 = VectorLoad((float *)P1);
 		VectorRegister4f R2 = VectorLoad((float *)P2);
 		VectorRegister4f Res = fn(BatchState, R0, R1, R2);
-		VectorStoreAligned(Res, (float *)P3);
+		VectorStore(Res, (float *)P3);
 	}
 	return InsPtr + 9;
 }
@@ -2039,7 +2039,7 @@ VM_FORCEINLINE const uint8 *VVM_Dispatch_execFn4f_1f(const bool CT_MultipleLoops
 			P2 += Inc2;
 			P3 += Inc3;
 			VectorRegister4f Res = fn(BatchState, R0, R1, R2, R3);
-			VectorStoreAligned(Res, (float *)P4);
+			VectorStore(Res, (float *)P4);
 			P4 += sizeof(FVecReg);
 		} while (P4 < End);
 	}
@@ -2050,7 +2050,7 @@ VM_FORCEINLINE const uint8 *VVM_Dispatch_execFn4f_1f(const bool CT_MultipleLoops
 		VectorRegister4f R2 = VectorLoad((float *)P2);
 		VectorRegister4f R3 = VectorLoad((float *)P3);
 		VectorRegister4f Res = fn(BatchState, R0, R1, R2, R3);
-		VectorStoreAligned(Res, (float *)P4);
+		VectorStore(Res, (float *)P4);
 	}
 	return InsPtr + 11;
 }
@@ -2085,7 +2085,7 @@ VM_FORCEINLINE const uint8 *VVM_Dispatch_execFn5f_1f(const bool CT_MultipleLoops
 			P3 += Inc3;
 			P4 += Inc4;
 			VectorRegister4f Res = fn(BatchState, R0, R1, R2, R3, R4);
-			VectorStoreAligned(Res, (float *)P5);
+			VectorStore(Res, (float *)P5);
 			P5 += sizeof(FVecReg);
 		} while (P5 < End);
 	}
@@ -2097,7 +2097,7 @@ VM_FORCEINLINE const uint8 *VVM_Dispatch_execFn5f_1f(const bool CT_MultipleLoops
 		VectorRegister4f R3 = VectorLoad((float *)P3);
 		VectorRegister4f R4 = VectorLoad((float *)P4);
 		VectorRegister4f Res = fn(BatchState, R0, R1, R2, R3, R4);
-		VectorStoreAligned(Res, (float *)P5);
+		VectorStore(Res, (float *)P5);
 	}
 	return InsPtr + 13;
 }
@@ -2122,8 +2122,8 @@ VM_FORCEINLINE const uint8 *VVM_Dispatch_execFn2f_2f(const bool CT_MultipleLoops
 			P0 += Inc0;
 			P1 += Inc1;
 			VectorRegister4f Res = fn(BatchState, R0, R1);
-			VectorStoreAligned(Res, (float *)P2);
-			VectorStoreAligned(Res, (float *)(P2 + P23d));
+			VectorStore(Res, (float *)P2);
+			VectorStore(Res, (float *)(P2 + P23d));
 			P2 += sizeof(FVecReg);
 		} while (P2 < End);
 	}
@@ -2132,8 +2132,8 @@ VM_FORCEINLINE const uint8 *VVM_Dispatch_execFn2f_2f(const bool CT_MultipleLoops
 		VectorRegister4f R0 = VectorLoad((float *)P0);
 		VectorRegister4f R1 = VectorLoad((float *)P1);
 		VectorRegister4f Res = fn(BatchState, R0, R1);
-		VectorStoreAligned(Res, (float *)P2);
-		VectorStoreAligned(Res, (float *)P3);
+		VectorStore(Res, (float *)P2);
+		VectorStore(Res, (float *)P3);
 	}
 	return InsPtr + 9;
 }
@@ -2152,7 +2152,7 @@ VM_FORCEINLINE const uint8 *VVM_Dispatch_execFn1i_1i(const bool CT_MultipleLoops
 			VectorRegister4i R0 = VectorIntLoad(P0);
 			P0 += Inc0;
 			VectorRegister4i Res = fn(BatchState, R0);
-			VectorIntStoreAligned(Res, P1);
+			VectorIntStore(Res, P1);
 			P1 += sizeof(FVecReg);
 		} while (P1 < End);
 	}
@@ -2160,7 +2160,7 @@ VM_FORCEINLINE const uint8 *VVM_Dispatch_execFn1i_1i(const bool CT_MultipleLoops
 	{
 		VectorRegister4i R0 = VectorIntLoad(P0);
 		VectorRegister4i Res = fn(BatchState, R0);
-		VectorIntStoreAligned(Res, P1);
+		VectorIntStore(Res, P1);
 	}
 	return InsPtr + 5;
 }
@@ -2183,7 +2183,7 @@ VM_FORCEINLINE const uint8 *VVM_Dispatch_execFn2i_1i(const bool CT_MultipleLoops
 			P0 += Inc0;
 			P1 += Inc1;
 			VectorRegister4i Res = fn(BatchState, R0, R1);
-			VectorIntStoreAligned(Res, P2);
+			VectorIntStore(Res, P2);
 			P2 += sizeof(FVecReg);
 		} while (P2 < End);
 	}
@@ -2192,7 +2192,7 @@ VM_FORCEINLINE const uint8 *VVM_Dispatch_execFn2i_1i(const bool CT_MultipleLoops
 		VectorRegister4i R0 = VectorIntLoad(P0);
 		VectorRegister4i R1 = VectorIntLoad(P1);
 		VectorRegister4i Res = fn(BatchState, R0, R1);
-		VectorIntStoreAligned(Res, P2);
+		VectorIntStore(Res, P2);
 	}
 	return InsPtr + 7;
 }
@@ -2219,7 +2219,7 @@ VM_FORCEINLINE const uint8 *VVM_Dispatch_execFn3i_1i(const bool CT_MultipleLoops
 			P1 += Inc1;
 			P2 += Inc2;
 			VectorRegister4i Res = fn(BatchState, R0, R1, R2);
-			VectorIntStoreAligned(Res, P3);
+			VectorIntStore(Res, P3);
 			P3 += sizeof(FVecReg);
 		} while (P3 < End);
 	}
@@ -2229,7 +2229,7 @@ VM_FORCEINLINE const uint8 *VVM_Dispatch_execFn3i_1i(const bool CT_MultipleLoops
 		VectorRegister4i R1 = VectorIntLoad(P1);
 		VectorRegister4i R2 = VectorIntLoad(P2);
 		VectorRegister4i Res = fn(BatchState, R0, R1, R2);
-		VectorIntStoreAligned(Res, P3);
+		VectorIntStore(Res, P3);
 	}
 	return InsPtr + 9;
 }
@@ -2249,8 +2249,8 @@ VM_FORCEINLINE const uint8 *VVM_Dispatch_execFn1i_2i(const bool CT_MultipleLoops
 			VectorRegister4i R0 = VectorIntLoad(P0);
 			P0 += Inc0;
 			VectorRegister4i Res = fn(BatchState, R0);
-			VectorIntStoreAligned(Res, P1);
-			VectorIntStoreAligned(Res, P2);
+			VectorIntStore(Res, P1);
+			VectorIntStore(Res, P2);
 			P1 += sizeof(FVecReg);
 			P2 += sizeof(FVecReg);
 		} while (P2 < End);
@@ -2259,8 +2259,8 @@ VM_FORCEINLINE const uint8 *VVM_Dispatch_execFn1i_2i(const bool CT_MultipleLoops
 	{
 		VectorRegister4i R0 = VectorIntLoad(P0);
 		VectorRegister4i Res = fn(BatchState, R0);
-		VectorIntStoreAligned(Res, P1);
-		VectorIntStoreAligned(Res, P2);
+		VectorIntStore(Res, P1);
+		VectorIntStore(Res, P2);
 	}
 	return InsPtr + 7;
 }
