Fix deterministic output for EVVM byte code

diff --git a/Engine/Source/Runtime/VectorVM/Private/VectorVMExperimentalOptimizer.cpp b/Engine/Source/Runtime/VectorVM/Private/VectorVMExperimentalOptimizer.cpp
--- a/Engine/Source/Runtime/VectorVM/Private/VectorVMExperimentalOptimizer.cpp
+++ b/Engine/Source/Runtime/VectorVM/Private/VectorVMExperimentalOptimizer.cpp
@@ -124,17 +124,17 @@ void FreezeVectorVMOptimizeContext(const FVectorVMOptimizeContext& Context, TArr
 	DestContext.OutputRemapDst        = nullptr;
 	DestContext.ExtFnTable            = nullptr;
 
-	// callbacks
-	DestContext.Init.ReallocFn = nullptr;
-	DestContext.Init.FreeFn = nullptr;
-	DestContext.Error.CallbackFn = nullptr;
-
 	// external function pointers
 	FVectorVMExtFunctionData* ExtFunctionTable = reinterpret_cast<FVectorVMExtFunctionData*>(DestData + ExtFnOffset);
 	for (uint32 ExtFunctionIt = 0; ExtFunctionIt < Context.NumExtFns; ++ExtFunctionIt)
 	{
 		ExtFunctionTable[ExtFunctionIt].Function = nullptr;
 	}
+
+	// clear out unnecessary elements - todo - consider just moving to a raw dump of what we actually care about
+	FMemory::Memzero(DestContext.Init);
+	FMemory::Memzero(DestContext.Error);
+	FMemory::Memzero(DestContext.Intermediate);
 }
 
 static void* VectorVMFrozenRealloc(void* Ptr, size_t NumBytes, const char* Filename, int LineNum)
