New VectorVM half float and cast instruction fixes.

diff --git a/Engine/Source/Runtime/Core/Public/Math/UnrealMathNeon.h b/Engine/Source/Runtime/Core/Public/Math/UnrealMathNeon.h
--- a/Engine/Source/Runtime/Core/Public/Math/UnrealMathNeon.h
+++ b/Engine/Source/Runtime/Core/Public/Math/UnrealMathNeon.h
@@ -587,6 +587,11 @@ FORCEINLINE VectorRegister4Double VectorLoadDouble1(const double* Ptr)
 	return Result;
 }
 
+FORCEINLINE VectorRegister4i VectorLoad64Bits(const void *Ptr)
+{
+	return vcombine_s64(vld1_s64((const int64_t *)Ptr), vdup_n_s64(0));
+}
+
 /**
  * Loads 4 unaligned floats - 2 from the first pointer, 2 from the second, and packs
  * them in to 1 vector.
@@ -2837,6 +2842,7 @@ FORCEINLINE VectorRegister4Int VectorFloatToInt(const VectorRegister4Double& A)
 * @return		VectorRegister4Int(Ptr[0], Ptr[1], Ptr[2], Ptr[3])
 */
 #define VectorIntLoad( Ptr )				vld1q_s32( (int32*)((void*)(Ptr)) )
+#define VectorIntLoad_16( Ptr )				vld1q_s16( (int16*)((void*)(Ptr)) )
 
 /**
 * Stores a vector to memory (aligned).
diff --git a/Engine/Source/Runtime/Core/Public/Math/UnrealMathSSE.h b/Engine/Source/Runtime/Core/Public/Math/UnrealMathSSE.h
--- a/Engine/Source/Runtime/Core/Public/Math/UnrealMathSSE.h
+++ b/Engine/Source/Runtime/Core/Public/Math/UnrealMathSSE.h
@@ -648,6 +648,11 @@ FORCEINLINE VectorRegister4Double VectorLoadDouble1(const double* Ptr)
 	return Result;
 }
 
+FORCEINLINE VectorRegister4i VectorLoad64Bits(const VectorRegister4i *Ptr)
+{
+	return _mm_loadu_si64((__m128i *)Ptr);
+}
+
 /**
  * Loads 2 floats from unaligned memory into X and Y and duplicates them in Z and W.
  *
@@ -3456,7 +3461,7 @@ FORCEINLINE VectorRegister4Int VectorFloatToInt(const VectorRegister4Double& A)
 * @return		VectorRegister4Int(Ptr[0], Ptr[1], Ptr[2], Ptr[3])
 */
 #define VectorIntLoad( Ptr )				_mm_loadu_si128( (VectorRegister4Int*)(Ptr) )
-
+#define VectorIntLoad_16( Ptr )             _mm_loadu_si64 ( (VectorRegister4Int*)(Ptr) )
 /**
 * Stores a vector to memory (aligned).
 *
diff --git a/Engine/Source/Runtime/VectorVM/Private/VectorVMExperimental.cpp b/Engine/Source/Runtime/VectorVM/Private/VectorVMExperimental.cpp
--- a/Engine/Source/Runtime/VectorVM/Private/VectorVMExperimental.cpp
+++ b/Engine/Source/Runtime/VectorVM/Private/VectorVMExperimental.cpp
@@ -5,7 +5,6 @@
 #include "VectorVM.h"
 #include "HAL/ConsoleManager.h"
 
-
 //prototypes for internal functions to the VVM
 #if VECTORVM_SUPPORTS_SERIALIZATION
 void VVMSer_serializeInstruction(FVectorVMSerializeState *SerializeState, FVectorVMSerializeState *CmpSerializeState, 
@@ -18,6 +17,25 @@ int VVMGetRegisterType(FVectorVMState *VVMState, uint16 RegIdx, uint16 *OutAbsRe
 uint32 *VVMSer_getRegPtrTablePtrFromIns(FVectorVMSerializeState *SerializeState, FVectorVMSerializeInstruction *Ins, uint16 RegIdx);
 #endif
 
+#if VECTORVM_DEBUG_PRINTF
+
+#define VECTORVM_PRINTF(fmt, ...)	                    \
+	{                                                   \
+		char buff[1024];                                \
+		snprintf(buff, sizeof(buff), fmt, __VA_ARGS__); \
+		wchar_t buff16[1024];                           \
+		char *p = buff;                                 \
+		wchar_t *p16 = buff16;                          \
+		while (*p) *p16++ = *p++;                       \
+		*p16 = 0;                                       \
+		UE_LOG(LogVectorVM, Warning, TEXT("%s"), (TCHAR *)buff16);  \
+	}
+
+#define VECTORVM_PRINT_NEW_LINE()
+
+int VVMGetRegisterType(FVectorVMState *VVMState, uint16 RegIdx, uint16 *OutAbsReg);
+#endif
+
 void *VVMDefaultRealloc(void *Ptr, size_t NumBytes, const char *Filename, int LineNumber)
 {
 	return FMemory::Realloc(Ptr, NumBytes);
@@ -230,7 +248,7 @@ void VVMMemSet32(void *dst, uint32 val, size_t num_vals)
 		char *RESTRICT end_ptr = ptr + num_vals * 4 - 16;
 		while (ptr < end_ptr) {
 			VectorIntStore(v4, ptr);
-			ptr += 16;
+			ptr += sizeof(v4);
 		}
 		VectorIntStore(v4, end_ptr);
 	}
@@ -269,12 +287,13 @@ void VVMMemSet16(void *dst, uint16 val, size_t num_vals)
 		char *RESTRICT end_ptr = ptr + num_vals * 4 - 16;
 		while (ptr < end_ptr) {
 			VectorIntStore(Val4, ptr);
-			ptr += 16;
+			ptr += sizeof(Val4);
 		}
 		VectorIntStore(Val4, end_ptr);
 	}
 }
 
+
 #if PLATFORM_CPU_X86_FAMILY
 #define VVM_pshufb(Src, Mask) _mm_shuffle_epi8(Src, Mask)
   // Fabian's round-to-nearest-even float to half
@@ -748,26 +767,133 @@ static void VVMBuildMapTableCaches(FVectorVMExecContext *ExecCtx)
 	}
 }
 
-static const uint8 *VVM_Output16(const bool CT_MultipleLoops, const uint8 *InsPtr, FVectorVMBatchState *BatchState, FVectorVMExecContext *ExecCtx)
-{
-	uint32 *NumOutputPerDataSet         = BatchState->ChunkLocalData.NumOutputPerDataSet;
-	uint32 *StartingOutputIdxPerDataSet = BatchState->ChunkLocalData.StartingOutputIdxPerDataSet;
-	uint8 **RegPtrTable                 = BatchState->RegPtrTable;
-	uint8 *RegIncTable                  = BatchState->RegIncTable;
-	uint8 **OutputMaskIdx               = BatchState->ChunkLocalData.OutputMaskIdx;
-	TArrayView<FDataSetMeta> DataSets   = ExecCtx->DataSets;
-
-	uint8  RegType            = InsPtr[-1] - (uint8)EVectorVMOp::outputdata_float;
-	int NumOutputLoops        = InsPtr[0];
-	uint8 DataSetIdx          = InsPtr[1];
-	uint32 NumOutputInstances = NumOutputPerDataSet[DataSetIdx];
-	uint32 InstanceOffset     = 0;
-	uint32 RegTypeOffset      = DataSets[DataSetIdx].OutputRegisterTypeOffsets[RegType];
-
-	const uint16 * RESTRICT SrcIndices       = (uint16 *)(InsPtr + 2);
-	const uint16 * RESTRICT DstIndices       = SrcIndices + NumOutputLoops;
+#define VVM_OUTPUT_FUNCTION_HEADER(CT_InputInsOutputTypeOpCode)                                              \
+	uint32 *NumOutputPerDataSet         = BatchState->ChunkLocalData.NumOutputPerDataSet;                    \
+	uint32 *StartingOutputIdxPerDataSet = BatchState->ChunkLocalData.StartingOutputIdxPerDataSet;            \
+	uint8 **RegPtrTable                 = BatchState->RegPtrTable;                                           \
+	uint8 *RegIncTable                  = BatchState->RegIncTable;                                           \
+	uint8 **OutputMaskIdx               = BatchState->ChunkLocalData.OutputMaskIdx;                          \
+	TArrayView<FDataSetMeta> DataSets   = ExecCtx->DataSets;                                                 \
+	uint8  RegType            = (uint8)(CT_InputInsOutputTypeOpCode) - (uint8)EVectorVMOp::outputdata_float; \
+	int NumOutputLoops        = InsPtr[0];                                                                   \
+	uint8 DataSetIdx          = InsPtr[1];                                                                   \
+	uint32 NumOutputInstances = NumOutputPerDataSet[DataSetIdx];                                             \
+	uint32 InstanceOffset     = 0;                                                                           \
+	uint32 RegTypeOffset      = DataSets[DataSetIdx].OutputRegisterTypeOffsets[RegType];                     \
+	const uint16 * RESTRICT SrcIndices       = (uint16 *)(InsPtr + 2);                                       \
+	const uint16 * RESTRICT DstIndices       = SrcIndices + NumOutputLoops;                                  \
 	InsPtr += 3 + 4 * NumOutputLoops;
 
+
+#if VECTORVM_DEBUG_PRINTF
+
+#	define VVM_OUTPUT_FUNCTION_FOOTER                                                                             \
+		VECTORVM_PRINTF("\tNum Loops: %d, DataSetIndex: %d\n", NumOutputLoops, DataSetIdx);                       \
+		for (int j = 0; j < NumOutputLoops; ++j)                                                                  \
+		{                                                                                                         \
+			uint16 AbsRegIdx;                                                                                     \
+			int ThisRegType = VVMGetRegisterType(ExecCtx->VVMState, SrcIndices[j], &AbsRegIdx);                   \
+			uint32 *DstReg = (uint32 *)RegPtrTable[DstIndices[j]] + InstanceOffset;                               \
+			VECTORVM_PRINTF("\t%c%d -> %d [0x%p]\n", VVM_RT_CHAR[ThisRegType], AbsRegIdx, DstIndices[j], DstReg); \
+		}                                                                                                         \
+		return InsPtr
+
+#else //VECTORVM_DEBUG_PRINTF
+
+#	define VVM_OUTPUT_FUNCTION_FOOTER return InsPtr
+
+#endif //VECTORVM_DEBUG_PRINTF
+
+static const uint8 *VVM_Output32_from_16(const bool CT_MultipleLoops, const uint8 *InsPtr, FVectorVMBatchState *BatchState, FVectorVMExecContext *ExecCtx)
+{
+	VVM_OUTPUT_FUNCTION_HEADER(EVectorVMOp::outputdata_float);
+	VVM_OUTPUT_FUNCTION_FOOTER;
+}
+
+static const uint8 *VVM_Output16_from_16(const bool CT_MultipleLoops, const uint8 *InsPtr, FVectorVMBatchState *BatchState, FVectorVMExecContext *ExecCtx)
+{
+	VVM_OUTPUT_FUNCTION_HEADER(EVectorVMOp::outputdata_half);
+	if (CT_MultipleLoops)
+	{
+		if (NumOutputInstances == BatchState->ChunkLocalData.NumInstancesThisChunk) //all outputs written
+		{ 
+			for (int j = 0; j < NumOutputLoops; ++j)
+			{
+				int      SrcInc = RegIncTable[SrcIndices[j]];
+				uint32 * SrcReg = (uint32 *)RegPtrTable[SrcIndices[j]];
+				uint32 * DstReg = (uint32 *)RegPtrTable[DstIndices[j]] + InstanceOffset;
+				if (SrcReg != DstReg)
+				{ //temp registers can be aliased to outputs
+					if (SrcInc == 0) //setting from a constant
+					{ 
+						VVMMemSet16(DstReg, *SrcReg, NumOutputInstances);
+					}
+					else
+					{
+						VVMMemCpy(DstReg, SrcReg, sizeof(uint16) * NumOutputInstances);
+					}
+				}
+			}
+		}
+		else if (NumOutputInstances > 0) //not all outputs are being written
+		{
+			for (int j = 0; j < NumOutputLoops; ++j)
+			{
+				int     SrcInc = RegIncTable[SrcIndices[j]];
+				uint64 *DstReg = (uint64 *)RegPtrTable[DstIndices[j]] + InstanceOffset;
+
+				if (SrcInc == 0) //setting from a constant
+				{ 
+					uint64 Val              = *(uint64 *)RegPtrTable[SrcIndices[j]];
+					uint64 *RESTRICT DstEnd = DstReg + NumOutputInstances;
+					uint64 *RESTRICT Ptr    = DstReg;
+					while (Ptr < DstEnd) {
+						*Ptr = Val;
+						Ptr++;
+					}
+				}
+				else
+				{
+					int NumLoops = (int)((BatchState->ChunkLocalData.NumInstancesThisChunk + 3) & ~3) >> 2; //assumes 4-wide ops
+					char * SrcPtr                 = (char *)RegPtrTable[SrcIndices[j]]; //src and dst can alias
+					char * DstPtr                 = (char *)DstReg;
+					uint8 * RESTRICT TblIdxPtr    = OutputMaskIdx[DataSetIdx];
+					uint8 * RESTRICT TblIdxEndPtr = TblIdxPtr + NumLoops;
+					while (TblIdxPtr < TblIdxEndPtr)
+					{
+						
+						uint8 TblIdx = *TblIdxPtr++;
+						VectorRegister4i Mask = ((VectorRegister4i *)VVM_PSHUFB_OUTPUT_TABLE16)[TblIdx];
+						VectorRegister4i Src  = VectorIntLoad_16(SrcPtr);
+						VectorRegister4i Val  = VVM_pshufb(Src, Mask);
+						VectorIntStore_16(Val, DstPtr);
+						SrcPtr += sizeof(uint16) * 4;
+						DstPtr += VVM_OUTPUT_ADVANCE_TABLE16[TblIdx];
+					}
+				}
+			}
+		}
+	}
+	else
+	{
+		uint8 OutputMask = OutputMaskIdx[DataSetIdx][0];
+		for (int j = 0; j < NumOutputLoops; ++j)
+		{
+			uint32 *SrcReg = (uint32 *)RegPtrTable[SrcIndices[j]];
+			uint32 *DstReg = (uint32 *)RegPtrTable[DstIndices[j]] + InstanceOffset;
+			VectorRegister4i Mask = ((VectorRegister4i *)VVM_PSHUFB_OUTPUT_TABLE16)[OutputMask];
+			VectorRegister4i Src  = VectorIntLoad(SrcReg);
+			VectorRegister4i Val  = VVM_pshufb(Src, Mask);
+			VectorIntStore_16(Val, DstReg);
+		}
+	}
+	VVM_OUTPUT_FUNCTION_FOOTER;
+}
+
+static const uint8 *VVM_Output16(const bool CT_MultipleLoops, const uint8 *InsPtr, FVectorVMBatchState *BatchState, FVectorVMExecContext *ExecCtx)
+{
+	VVM_OUTPUT_FUNCTION_HEADER(EVectorVMOp::outputdata_half);
+	
 	if (CT_MultipleLoops)
 	{
 		if (NumOutputInstances == BatchState->ChunkLocalData.NumInstancesThisChunk) //all outputs written
@@ -802,7 +928,7 @@ static const uint8 *VVM_Output16(const bool CT_MultipleLoops, const uint8 *InsPt
 			for (int j = 0; j < NumOutputLoops; ++j)
 			{
 				int   SrcInc = RegIncTable[SrcIndices[j]];
-				char *DstReg = (char *)RegPtrTable[DstIndices[j]] + InstanceOffset;
+				char *DstReg = (char *)RegPtrTable[DstIndices[j]] + (InstanceOffset * sizeof(uint16));
 				if (SrcInc == 0) //setting from a constant
 				{ 
 					uint16 Val;
@@ -839,48 +965,17 @@ static const uint8 *VVM_Output16(const bool CT_MultipleLoops, const uint8 *InsPt
 			uint16 *DstReg = (uint16 *)RegPtrTable[DstIndices[j]] + InstanceOffset;
 			//convert 4 values at once then shift them in place
 			VectorRegister4i Mask = ((VectorRegister4i *)VVM_PSHUFB_OUTPUT_TABLE16)[OutputMask];
-			//VectorRegister4i Src4 = _mm_cvtps_ph(_mm_loadu_ps((float *)SrcReg), _MM_FROUND_TO_NEAREST_INT);
-			//_mm_storeu_si64(DstReg, _mm_shuffle_epi8(Src4, Mask));
-			
 			VectorRegister4i HalfVals;
 			VVM_floatToHalf(&HalfVals, SrcReg);
 			VectorIntStore_16(VVM_pshufb(HalfVals, Mask), DstReg);
 		}
 	}
-#	if VECTORVM_DEBUG_PRINTF
-	printf("\tNum Loops: %d, DataSetIndex: %d\n", NumOutputLoops, DataSetIdx);
-	for (int j = 0; j < NumOutputLoops; ++j)
-	{
-		uint16 AbsRegIdx;
-		int RegType = VVMGetRegisterType(ExecCtx->VVMState, SrcIndices[j], &AbsRegIdx);
-		uint32 *DstReg = (uint32 *)RegPtrTable[DstIndices[j]] + InstanceOffset;
-		printf("\t%c%d -> %d [0x%p]\n", VVM_RT_CHAR[RegType], AbsRegIdx, DstIndices[j], DstReg);
-	}
-	fflush(stdout);
-#	endif //VECTORVM_DEBUG_PRINTF
-	return InsPtr;
+	VVM_OUTPUT_FUNCTION_FOOTER;
 }
 
 static const uint8 *VVM_Output32(const bool CT_MultipleLoops, const uint8 *InsPtr, FVectorVMBatchState *BatchState, FVectorVMExecContext *ExecCtx)
 {
-	uint32 *NumOutputPerDataSet         = BatchState->ChunkLocalData.NumOutputPerDataSet;
-	uint32 *StartingOutputIdxPerDataSet = BatchState->ChunkLocalData.StartingOutputIdxPerDataSet;
-	uint8 **RegPtrTable                 = BatchState->RegPtrTable;
-	uint8 *RegIncTable                  = BatchState->RegIncTable;
-	uint8 **OutputMaskIdx               = BatchState->ChunkLocalData.OutputMaskIdx;
-	TArrayView<FDataSetMeta> DataSets   = ExecCtx->DataSets;
-
-	uint8  RegType            = InsPtr[-1] - (uint8)EVectorVMOp::outputdata_float;
-	int NumOutputLoops        = InsPtr[0];
-	uint8 DataSetIdx          = InsPtr[1];
-	uint32 NumOutputInstances = NumOutputPerDataSet[DataSetIdx];
-	uint32 InstanceOffset     = 0;
-	uint32 RegTypeOffset      = DataSets[DataSetIdx].OutputRegisterTypeOffsets[RegType];
-
-	const uint16 * RESTRICT SrcIndices       = (uint16 *)(InsPtr + 2);
-	const uint16 * RESTRICT DstIndices       = SrcIndices + NumOutputLoops;
-	InsPtr += 3 + 4 * NumOutputLoops;
-
+	VVM_OUTPUT_FUNCTION_HEADER(InsPtr[-1]);
 	if (CT_MultipleLoops)
 	{
 		if (NumOutputInstances == BatchState->ChunkLocalData.NumInstancesThisChunk) //all outputs written
@@ -956,19 +1051,7 @@ static const uint8 *VVM_Output32(const bool CT_MultipleLoops, const uint8 *InsPt
 			VectorIntStore(Val, DstReg);
 		}
 	}
-#	if VECTORVM_DEBUG_PRINTF
-	printf("\tNum Loops: %d, DataSetIndex: %d\n", NumOutputLoops, DataSetIdx);
-	for (int j = 0; j < NumOutputLoops; ++j)
-	{
-		uint16 AbsSrcRegIdx;
-		uint16 AbsDstRegIdx;
-		int RegType = VVMGetRegisterType(ExecCtx->VVMState, SrcIndices[j], &AbsSrcRegIdx);
-		VVMGetRegisterType(ExecCtx->VVMState, DstIndices[j], &AbsDstRegIdx);
-		uint32 *DstReg = (uint32 *)RegPtrTable[DstIndices[j]] + InstanceOffset;
-		printf("\t%c%d -> O%d (%d) [0x%p]\n", VVM_RT_CHAR[RegType], AbsSrcRegIdx, AbsDstRegIdx, DstIndices[j], DstReg);
-	}
-#	endif //VECTORVM_DEBUG_PRINTF
-	return InsPtr;
+	VVM_OUTPUT_FUNCTION_FOOTER;
 }
 
 static const uint8 *VVM_acquireindex(const bool CT_MultipleLoops, const uint8 *InsPtr, FVectorVMBatchState *BatchState, FVectorVMExecContext *ExecCtx)
@@ -1656,7 +1739,7 @@ static const uint8 *VVM_acquire_id(const bool CT_MultipleLoops, const uint8 *Ins
 	return InsPtr + 6;
 }
 
-FORCEINLINE const uint8 *VVM_external_func_call(const bool CT_MultipleLoops, const uint8 *InsPtr, FVectorVMBatchState *BatchState, FVectorVMExecContext *ExecCtx, FVectorVMSerializeState *SerializeState, FVectorVMSerializeState *CmpSerializeState, int NumLoops)
+const uint8 *VVM_external_func_call(const bool CT_MultipleLoops, const uint8 *InsPtr, FVectorVMBatchState *BatchState, FVectorVMExecContext *ExecCtx, FVectorVMSerializeState *SerializeState, FVectorVMSerializeState *CmpSerializeState, int NumLoops)
 {
 	int FnIdx = (int)*(uint16 *)(InsPtr);
 	FVectorVMExtFunctionData *ExtFnData = ExecCtx->VVMState->ExtFunctionTable + FnIdx;
@@ -1665,9 +1748,10 @@ FORCEINLINE const uint8 *VVM_external_func_call(const bool CT_MultipleLoops, con
 	for (int i = 0; i < ExtFnData->NumInputs + ExtFnData->NumOutputs; ++i) {
 		uint16 AbsRegIdx;
 		int RegType = VVMGetRegisterType(ExecCtx->VVMState, (((uint16 *)InsPtr) + 1)[i], &AbsRegIdx);
-		printf("%c%d, ", VVM_RT_CHAR[RegType], AbsRegIdx);
+		VECTORVM_PRINTF("%c%d, ", VVM_RT_CHAR[RegType], AbsRegIdx);
 	}
-	printf("\n");
+	VECTORVM_PRINT_NEW_LINE();
+	
 #endif //VECTORVM_DEBUG_PRINTF
 #if	VECTORVM_SUPPORTS_SERIALIZATION && !defined(VVM_SERIALIZE_NO_WRITE)
 	if (SerializeState && (SerializeState->Flags & VVMSer_SyncExtFns) && CmpSerializeState && SerializeState->NumInstances == CmpSerializeState->NumInstances && (CmpSerializeState->NumInstructions > SerializeState->NumInstructions))
@@ -1847,7 +1931,7 @@ FORCEINLINE const uint8 *VVM_external_func_call(const bool CT_MultipleLoops, con
 		ExtFnData->Function->Execute(ExtFnCtx);
 #		endif //VECTORVM_SUPPORTS_LEGACY
 	}
-	return InsPtr + 3 + 2 * (ExtFnData->NumInputs + ExtFnData->NumOutputs);;
+	return InsPtr + 3 + 2 * (ExtFnData->NumInputs + ExtFnData->NumOutputs);
 }
 
 static VM_FORCEINLINE const uint8 *VVM_sincos(const bool CT_MultipleLoops, const uint8 *InsPtr, FVectorVMBatchState *BatchState, int NumLoops)
@@ -2265,182 +2349,184 @@ VM_FORCEINLINE const uint8 *VVM_Dispatch_execFn1i_2i(const bool CT_MultipleLoops
 	return InsPtr + 7;
 }
 
-VM_FORCEINLINE VectorRegister4f VVM_Exec2f_add                         (FVectorVMBatchState *BatchState, VectorRegister4f a, VectorRegister4f b)                                                            { return VectorAdd(a, b); }
-VM_FORCEINLINE VectorRegister4f VVM_Exec2f_sub                         (FVectorVMBatchState *BatchState, VectorRegister4f a, VectorRegister4f b)                                                            { return VectorSubtract(a, b); }
-VM_FORCEINLINE VectorRegister4f VVM_Exec2f_mul                         (FVectorVMBatchState *BatchState, VectorRegister4f a, VectorRegister4f b)                                                            { return VectorMultiply(a, b); }
-VM_FORCEINLINE VectorRegister4f VVM_Exec2f_div                         (FVectorVMBatchState *BatchState, VectorRegister4f a, VectorRegister4f b)                                                            { return VectorSelect(VectorCompareGT(VectorAbs(b), VVM_m128Const(Epsilon)), VectorDivide(a, b), VectorZeroFloat()); }
-VM_FORCEINLINE VectorRegister4f VVM_Exec3f_mad                         (FVectorVMBatchState *BatchState, VectorRegister4f a, VectorRegister4f b, VectorRegister4f c)                                        { return VectorMultiplyAdd(a, b, c); }
-VM_FORCEINLINE VectorRegister4f VVM_Exec3f_lerp                        (FVectorVMBatchState *BatchState, VectorRegister4f a, VectorRegister4f b, VectorRegister4f c)                                        { return VectorLerp(a, b, c); }
-VM_FORCEINLINE VectorRegister4f VVM_Exec1f_rcp                         (FVectorVMBatchState *BatchState, VectorRegister4f a)                                                                                { return VectorSelect(VectorCompareGT(VectorAbs(a) , VVM_m128Const(Epsilon)), VectorReciprocalEstimate(a), VectorZeroFloat()); }
-VM_FORCEINLINE VectorRegister4f VVM_Exec1f_rsq                         (FVectorVMBatchState *BatchState, VectorRegister4f a)                                                                                { return VectorSelect(VectorCompareGT(a, VVM_m128Const(Epsilon)), VectorReciprocalSqrt(a), VectorZeroFloat()); }
-VM_FORCEINLINE VectorRegister4f VVM_Exec1f_sqrt                        (FVectorVMBatchState *BatchState, VectorRegister4f a)                                                                                { return VectorSelect(VectorCompareGT(a, VVM_m128Const(Epsilon)), VectorSqrt(a), VectorZeroFloat()); }
-VM_FORCEINLINE VectorRegister4f VVM_Exec1f_neg                         (FVectorVMBatchState *BatchState, VectorRegister4f a)                                                                                { return VectorNegate(a); }
-VM_FORCEINLINE VectorRegister4f VVM_Exec1f_abs                         (FVectorVMBatchState *BatchState, VectorRegister4f a)                                                                                { return VectorAbs(a); }
-VM_FORCEINLINE VectorRegister4f VVM_Exec1f_exp                         (FVectorVMBatchState *BatchState, VectorRegister4f a)                                                                                { return VectorExp(a); }
-VM_FORCEINLINE VectorRegister4f VVM_Exec1f_exp2                        (FVectorVMBatchState *BatchState, VectorRegister4f a)                                                                                { return VectorExp2(a); }
-VM_FORCEINLINE VectorRegister4f VVM_Exec1f_log                         (FVectorVMBatchState *BatchState, VectorRegister4f a)                                                                                { return VectorSelect(VectorCompareGT(a, VectorZeroFloat()), VectorLog(a), VectorZeroFloat()); }
-VM_FORCEINLINE VectorRegister4f VVM_Exec1f_log2                        (FVectorVMBatchState *BatchState, VectorRegister4f a)                                                                                { return VectorLog2(a); }
-VM_FORCEINLINE VectorRegister4f VVM_Exec1f_sin                         (FVectorVMBatchState *BatchState, VectorRegister4f a)                                                                                { return VectorSin(a); }
-VM_FORCEINLINE VectorRegister4f VVM_Exec1f_cos                         (FVectorVMBatchState *BatchState, VectorRegister4f a)                                                                                { return VectorCos(a); }
-VM_FORCEINLINE VectorRegister4f VVM_Exec1f_tan                         (FVectorVMBatchState *BatchState, VectorRegister4f a)                                                                                { return VectorTan(a); }
-VM_FORCEINLINE VectorRegister4f VVM_Exec1f_acos                        (FVectorVMBatchState *BatchState, VectorRegister4f a)                                                                                { return VectorATan2(VVM_Exec1f_sqrt(BatchState, VectorMultiply(VectorSubtract(VVM_m128Const(One), a), VectorAdd(VVM_m128Const(One), a))), a); }
-VM_FORCEINLINE VectorRegister4f VVM_Exec1f_asin                        (FVectorVMBatchState *BatchState, VectorRegister4f a)                                                                                { return VectorSubtract(VVM_m128Const(QuarterPi), VVM_Exec1f_acos(BatchState, a)); }
-VM_FORCEINLINE VectorRegister4f VVM_Exec1f_atan                        (FVectorVMBatchState *BatchState, VectorRegister4f a)                                                                                { return VectorATan(a); }
-VM_FORCEINLINE VectorRegister4f VVM_Exec2f_atan2                       (FVectorVMBatchState *BatchState, VectorRegister4f a, VectorRegister4f b)                                                            { return VectorATan2(a, b); }
-VM_FORCEINLINE VectorRegister4f VVM_Exec1f_ceil                        (FVectorVMBatchState *BatchState, VectorRegister4f a)                                                                                { return VectorCeil(a); }
-VM_FORCEINLINE VectorRegister4f VVM_Exec1f_floor                       (FVectorVMBatchState *BatchState, VectorRegister4f a)                                                                                { return VectorFloor(a); }
-VM_FORCEINLINE VectorRegister4f VVM_Exec2f_fmod                        (FVectorVMBatchState *BatchState, VectorRegister4f a, VectorRegister4f b)                                                            { return VectorMod(a, b); }
-VM_FORCEINLINE VectorRegister4f VVM_Exec1f_frac                        (FVectorVMBatchState *BatchState, VectorRegister4f a)                                                                                { return VectorFractional(a); }
-VM_FORCEINLINE VectorRegister4f VVM_Exec1f_trunc                       (FVectorVMBatchState *BatchState, VectorRegister4f a)                                                                                { return VectorTruncate(a); }
-VM_FORCEINLINE VectorRegister4f VVM_Exec3f_clamp                       (FVectorVMBatchState *BatchState, VectorRegister4f a, VectorRegister4f b, VectorRegister4f c)                                        { return VectorClamp(a, b, c); }
-VM_FORCEINLINE VectorRegister4f VVM_Exec2f_min                         (FVectorVMBatchState *BatchState, VectorRegister4f a, VectorRegister4f b)                                                            { return VectorMin(a, b); }
-VM_FORCEINLINE VectorRegister4f VVM_Exec2f_max                         (FVectorVMBatchState *BatchState, VectorRegister4f a, VectorRegister4f b)                                                            { return VectorMax(a, b); }
-VM_FORCEINLINE VectorRegister4f VVM_Exec2f_pow                         (FVectorVMBatchState *BatchState, VectorRegister4f a, VectorRegister4f b)                                                            { return VectorSelect(VectorCompareGT(a, VVM_m128Const(Epsilon)), VectorPow(a, b), VectorZeroFloat()); }
-VM_FORCEINLINE VectorRegister4f VVM_Exec1f_round                       (FVectorVMBatchState *BatchState, VectorRegister4f a)                                                                                { return VectorRound(a); }
-VM_FORCEINLINE VectorRegister4f VVM_Exec1f_sign                        (FVectorVMBatchState *BatchState, VectorRegister4f a)                                                                                { return VectorSign(a); }
-VM_FORCEINLINE VectorRegister4f VVM_Exec2f_step                        (FVectorVMBatchState *BatchState, VectorRegister4f a, VectorRegister4f b)                                                            { return VectorStep(VectorSubtract(a, b)); }
-VM_FORCEINLINE const uint8 *    VVM_Exec1null_random                   (VVM_NULL_FN_ARGS)                                                                                                                   { return VVM_random(CT_MultipleLoops, InsPtr, BatchState, ExecCtx, SerializeState, CmpSerializeState, NumLoops); }
-VM_FORCEINLINE const uint8 *    VVM_Exec0null_noise                    (VVM_NULL_FN_ARGS)                                                                                                                   { return InsPtr; }
-VM_FORCEINLINE VectorRegister4f VVM_Exec2f_cmplt                       (FVectorVMBatchState *BatchState, VectorRegister4f a, VectorRegister4f b)                                                            { return VectorCompareLT(a, b); }
-VM_FORCEINLINE VectorRegister4f VVM_Exec2f_cmple                       (FVectorVMBatchState *BatchState, VectorRegister4f a, VectorRegister4f b)                                                            { return VectorCompareLE(a, b); }
-VM_FORCEINLINE VectorRegister4f VVM_Exec2f_cmpgt                       (FVectorVMBatchState *BatchState, VectorRegister4f a, VectorRegister4f b)                                                            { return VectorCompareGT(a, b); }
-VM_FORCEINLINE VectorRegister4f VVM_Exec2f_cmpge                       (FVectorVMBatchState *BatchState, VectorRegister4f a, VectorRegister4f b)                                                            { return VectorCompareGE(a, b); }
-VM_FORCEINLINE VectorRegister4f VVM_Exec2f_cmpeq                       (FVectorVMBatchState *BatchState, VectorRegister4f a, VectorRegister4f b)                                                            { return VectorCompareEQ(a, b); }
-VM_FORCEINLINE VectorRegister4f VVM_Exec2f_cmpneq                      (FVectorVMBatchState *BatchState, VectorRegister4f a, VectorRegister4f b)                                                            { return VectorCompareNE(a, b); }
-VM_FORCEINLINE VectorRegister4f VVM_Exec3f_select                      (FVectorVMBatchState *BatchState, VectorRegister4f a, VectorRegister4f b, VectorRegister4f c)                                        { return VectorSelect(a, b, c); }
-VM_FORCEINLINE VectorRegister4i VVM_Exec2i_addi                        (FVectorVMBatchState *BatchState, VectorRegister4i a, VectorRegister4i b)                                                            { return VectorIntAdd(a, b); }
-VM_FORCEINLINE VectorRegister4i VVM_Exec2i_subi                        (FVectorVMBatchState *BatchState, VectorRegister4i a, VectorRegister4i b)                                                            { return VectorIntSubtract(a, b); }
-VM_FORCEINLINE VectorRegister4i VVM_Exec2i_muli                        (FVectorVMBatchState *BatchState, VectorRegister4i a, VectorRegister4i b)                                                            { return VectorIntMultiply(a, b); }
-VM_FORCEINLINE VectorRegister4i VVM_Exec2i_divi                        (FVectorVMBatchState *BatchState, VectorRegister4i a, VectorRegister4i b)                                                            { return VVMIntDiv(a, b); }
-VM_FORCEINLINE VectorRegister4i VVM_Exec3i_clampi                      (FVectorVMBatchState *BatchState, VectorRegister4i a, VectorRegister4i b, VectorRegister4i c)                                        { return VectorIntClamp(a, b, c); }
-VM_FORCEINLINE VectorRegister4i VVM_Exec2i_mini                        (FVectorVMBatchState *BatchState, VectorRegister4i a, VectorRegister4i b)                                                            { return VectorIntMin(a, b); }
-VM_FORCEINLINE VectorRegister4i VVM_Exec2i_maxi                        (FVectorVMBatchState *BatchState, VectorRegister4i a, VectorRegister4i b)                                                            { return VectorIntMax(a, b); }
-VM_FORCEINLINE VectorRegister4i VVM_Exec1i_absi                        (FVectorVMBatchState *BatchState, VectorRegister4i a)                                                                                { return VectorIntAbs(a); }
-VM_FORCEINLINE VectorRegister4i VVM_Exec1i_negi                        (FVectorVMBatchState *BatchState, VectorRegister4i a)                                                                                { return VectorIntNegate(a); }
-VM_FORCEINLINE VectorRegister4i VVM_Exec1i_signi                       (FVectorVMBatchState *BatchState, VectorRegister4i a)                                                                                { return VectorIntSign(a); }
-VM_FORCEINLINE const uint8 *    VVM_Exec1null_randomi                  (VVM_NULL_FN_ARGS)                                                                                                                   { return VVM_randomi(CT_MultipleLoops, InsPtr, BatchState, ExecCtx, SerializeState, CmpSerializeState, NumLoops); }
-VM_FORCEINLINE VectorRegister4i VVM_Exec2i_cmplti                      (FVectorVMBatchState *BatchState, VectorRegister4i a, VectorRegister4i b)                                                            { return VectorIntCompareLT(a, b); }
-VM_FORCEINLINE VectorRegister4i VVM_Exec2i_cmplei                      (FVectorVMBatchState *BatchState, VectorRegister4i a, VectorRegister4i b)                                                            { return VectorIntCompareLE(a, b); }
-VM_FORCEINLINE VectorRegister4i VVM_Exec2i_cmpgti                      (FVectorVMBatchState *BatchState, VectorRegister4i a, VectorRegister4i b)                                                            { return VectorIntCompareGT(a, b); }
-VM_FORCEINLINE VectorRegister4i VVM_Exec2i_cmpgei                      (FVectorVMBatchState *BatchState, VectorRegister4i a, VectorRegister4i b)                                                            { return VectorIntCompareGE(a, b); }
-VM_FORCEINLINE VectorRegister4i VVM_Exec2i_cmpeqi                      (FVectorVMBatchState *BatchState, VectorRegister4i a, VectorRegister4i b)                                                            { return VectorIntCompareEQ(a, b); }
-VM_FORCEINLINE VectorRegister4i VVM_Exec2i_cmpneqi                     (FVectorVMBatchState *BatchState, VectorRegister4i a, VectorRegister4i b)                                                            { return VectorIntCompareNEQ(a, b); }
-VM_FORCEINLINE VectorRegister4i VVM_Exec2i_bit_and                     (FVectorVMBatchState *BatchState, VectorRegister4i a, VectorRegister4i b)                                                            { return VectorIntAnd(a, b); }
-VM_FORCEINLINE VectorRegister4i VVM_Exec2i_bit_or                      (FVectorVMBatchState *BatchState, VectorRegister4i a, VectorRegister4i b)                                                            { return VectorIntOr(a, b); }
-VM_FORCEINLINE VectorRegister4i VVM_Exec2i_bit_xor                     (FVectorVMBatchState *BatchState, VectorRegister4i a, VectorRegister4i b)                                                            { return VectorIntXor(a, b); }
-VM_FORCEINLINE VectorRegister4i VVM_Exec1i_bit_not                     (FVectorVMBatchState *BatchState, VectorRegister4i a)                                                                                { return VectorIntNot(a); }
-VM_FORCEINLINE VectorRegister4i VVM_Exec2i_bit_lshift                  (FVectorVMBatchState *BatchState, VectorRegister4i a, VectorRegister4i b)                                                            { return VVMIntLShift(a, b); }
-VM_FORCEINLINE VectorRegister4i VVM_Exec2i_bit_rshift                  (FVectorVMBatchState *BatchState, VectorRegister4i a, VectorRegister4i b)                                                            { return VVMIntRShift(a, b); }
-VM_FORCEINLINE VectorRegister4i VVM_Exec2i_logic_and                   (FVectorVMBatchState *BatchState, VectorRegister4i a, VectorRegister4i b)                                                            { return VectorIntAnd(a, b); }
-VM_FORCEINLINE VectorRegister4i VVM_Exec2i_logic_or                    (FVectorVMBatchState *BatchState, VectorRegister4i a, VectorRegister4i b)                                                            { return VectorIntOr(a, b); }
-VM_FORCEINLINE VectorRegister4i VVM_Exec2i_logic_xor                   (FVectorVMBatchState *BatchState, VectorRegister4i a, VectorRegister4i b)                                                            { return VectorIntXor(a, b); }
-VM_FORCEINLINE VectorRegister4i VVM_Exec1i_logic_not                   (FVectorVMBatchState *BatchState, VectorRegister4i a)                                                                                { return VectorIntNot(a); }
-VM_FORCEINLINE VectorRegister4i VVM_Exec1i_f2i                         (FVectorVMBatchState *BatchState, VectorRegister4i a)                                                                                { return VVMf2i(a); }
-VM_FORCEINLINE VectorRegister4f VVM_Exec1f_i2f                         (FVectorVMBatchState *BatchState, VectorRegister4f a)                                                                                { return VVMi2f(a); }
-VM_FORCEINLINE VectorRegister4f VVM_Exec1f_f2b                         (FVectorVMBatchState *BatchState, VectorRegister4f a)                                                                                { return VectorCompareGT(a, VectorZeroFloat()); }
-VM_FORCEINLINE VectorRegister4f VVM_Exec1f_b2f                         (FVectorVMBatchState *BatchState, VectorRegister4f a)                                                                                { return VectorSelect(a, VVM_m128Const(One), VectorZeroFloat()); }
-VM_FORCEINLINE VectorRegister4i VVM_Exec1i_i2b                         (FVectorVMBatchState *BatchState, VectorRegister4i a)                                                                                { return VectorIntCompareGT(a, VectorSetZero()); }
-VM_FORCEINLINE VectorRegister4i VVM_Exec1i_b2i                         (FVectorVMBatchState *BatchState, VectorRegister4i a)                                                                                { return VectorIntSelect(a, VectorIntSet1(1), VectorSetZero()); }
-VM_FORCEINLINE const uint8 *    VVM_Exec0null_outputdata_float         (VVM_NULL_FN_ARGS)                                                                                                                   { return VVM_Output32(CT_MultipleLoops, InsPtr, BatchState, ExecCtx); }
-VM_FORCEINLINE const uint8 *    VVM_Exec0null_outputdata_int32         (VVM_NULL_FN_ARGS)                                                                                                                   { return VVM_Output32(CT_MultipleLoops, InsPtr, BatchState, ExecCtx); }
-VM_FORCEINLINE const uint8 *    VVM_Exec0null_outputdata_half          (VVM_NULL_FN_ARGS)                                                                                                                   { return VVM_Output16(CT_MultipleLoops, InsPtr, BatchState, ExecCtx); }
-VM_FORCEINLINE const uint8 *    VVM_Exec1null_acquireindex             (VVM_NULL_FN_ARGS)                                                                                                                   { return VVM_acquireindex(CT_MultipleLoops, InsPtr, BatchState, ExecCtx); }
-VM_FORCEINLINE const uint8 *    VVM_Exec0null_external_func_call       (VVM_NULL_FN_ARGS)                                                                                                                   { return VVM_external_func_call(CT_MultipleLoops, InsPtr, BatchState, ExecCtx, SerializeState, CmpSerializeState, NumLoops); }
-VM_FORCEINLINE const uint8 *    VVM_Exec1null_exec_index               (VVM_NULL_FN_ARGS)                                                                                                                   { return VVM_exec_index(CT_MultipleLoops, InsPtr, BatchState, NumLoops); }
-VM_FORCEINLINE const uint8 *    VVM_Exec0null_noise2D                  (VVM_NULL_FN_ARGS)                                                                                                                   { return InsPtr; }
-VM_FORCEINLINE const uint8 *    VVM_Exec0null_noise3D                  (VVM_NULL_FN_ARGS)                                                                                                                   { return InsPtr; }
-VM_FORCEINLINE const uint8 *    VVM_Exec1null_enter_stat_scope         (VVM_NULL_FN_ARGS)                                                                                                                   { return InsPtr + 2; }
-VM_FORCEINLINE const uint8 *    VVM_Exec0null_exit_stat_scope          (VVM_NULL_FN_ARGS)                                                                                                                   { return InsPtr; }
-VM_FORCEINLINE const uint8 *    VVM_Exec0null_update_id                (VVM_NULL_FN_ARGS)                                                                                                                   { return VVM_update_id(CT_MultipleLoops, InsPtr, BatchState, ExecCtx, NumLoops); }
-VM_FORCEINLINE const uint8 *    VVM_Exec1null_acquire_id               (VVM_NULL_FN_ARGS)                                                                                                                   { return VVM_acquire_id(CT_MultipleLoops, InsPtr, BatchState, ExecCtx, NumLoops); }
-VM_FORCEINLINE const uint8 *    VVM_Exec1null_half_to_float            (VVM_NULL_FN_ARGS)                                                                                                                   { return VVM_half_to_float(CT_MultipleLoops, InsPtr, BatchState, NumLoops); }
-VM_FORCEINLINE const uint8 *    VVM_Exec1null_fasi                     (VVM_NULL_FN_ARGS)                                                                                                                   { return InsPtr; }
-VM_FORCEINLINE const uint8 *    VVM_Exec1null_iasf                     (VVM_NULL_FN_ARGS)                                                                                                                   { return InsPtr; }
-VM_FORCEINLINE const uint8 *    VVM_Exec1null_exec_indexf              (VVM_NULL_FN_ARGS)                                                                                                                   { return VVM_exec_indexf(CT_MultipleLoops, InsPtr, BatchState, NumLoops); }
-VM_FORCEINLINE const uint8 *    VVM_Exec1null_exec_index_addi          (VVM_NULL_FN_ARGS)                                                                                                                   { return VVM_exec_index_addi(CT_MultipleLoops, InsPtr, BatchState, NumLoops); }
-VM_FORCEINLINE VectorRegister4f VVM_Exec4f_cmplt_select                (FVectorVMBatchState *BatchState, VectorRegister4f a, VectorRegister4f b, VectorRegister4f c, VectorRegister4f d)                    { return VectorSelect(VectorCompareLT(a, b), c, d); }
-VM_FORCEINLINE VectorRegister4f VVM_Exec4f_cmple_select                (FVectorVMBatchState *BatchState, VectorRegister4f a, VectorRegister4f b, VectorRegister4f c, VectorRegister4f d)                    { return VectorSelect(VectorCompareLE(a, b), c, d); }
-VM_FORCEINLINE VectorRegister4f VVM_Exec4f_cmpeq_select                (FVectorVMBatchState *BatchState, VectorRegister4f a, VectorRegister4f b, VectorRegister4f c, VectorRegister4f d)                    { return VectorSelect(VectorCompareEQ(a, b), c, d); }
-VM_FORCEINLINE VectorRegister4f VVM_Exec4f_cmplti_select               (FVectorVMBatchState *BatchState, VectorRegister4f a, VectorRegister4f b, VectorRegister4f c, VectorRegister4f d)                    { return VectorSelect(VectorCastIntToFloat(VectorIntCompareLT(*(VectorRegister4i *)&a, *(VectorRegister4i *)&b)), c, d); }
-VM_FORCEINLINE VectorRegister4f VVM_Exec4f_cmplei_select               (FVectorVMBatchState *BatchState, VectorRegister4f a, VectorRegister4f b, VectorRegister4f c, VectorRegister4f d)                    { return VectorSelect(VectorCastIntToFloat(VectorIntCompareLE(*(VectorRegister4i *)&a, *(VectorRegister4i *)&b)), c, d); }
-VM_FORCEINLINE VectorRegister4f VVM_Exec4f_cmpeqi_select               (FVectorVMBatchState *BatchState, VectorRegister4f a, VectorRegister4f b, VectorRegister4f c, VectorRegister4f d)                    { return VectorSelect(VectorCastIntToFloat(VectorIntCompareEQ(*(VectorRegister4i *)&a, *(VectorRegister4i *)&b)), c, d); }
-VM_FORCEINLINE VectorRegister4f VVM_Exec3f_cmplt_logic_and             (FVectorVMBatchState *BatchState, VectorRegister4f a, VectorRegister4f b, VectorRegister4f c)                                        { return VectorCastIntToFloat(VectorIntAnd(VectorCastFloatToInt(VectorCompareLT(a, b)), *(VectorRegister4i *)&c)); }
-VM_FORCEINLINE VectorRegister4f VVM_Exec3f_cmple_logic_and             (FVectorVMBatchState *BatchState, VectorRegister4f a, VectorRegister4f b, VectorRegister4f c)                                        { return VectorCastIntToFloat(VectorIntAnd(VectorCastFloatToInt(VectorCompareLE(a, b)), *(VectorRegister4i *)&c)); }
-VM_FORCEINLINE VectorRegister4f VVM_Exec3f_cmpgt_logic_and             (FVectorVMBatchState *BatchState, VectorRegister4f a, VectorRegister4f b, VectorRegister4f c)                                        { return VectorCastIntToFloat(VectorIntAnd(VectorCastFloatToInt(VectorCompareGT(a, b)), *(VectorRegister4i *)&c)); }
-VM_FORCEINLINE VectorRegister4f VVM_Exec3f_cmpge_logic_and             (FVectorVMBatchState *BatchState, VectorRegister4f a, VectorRegister4f b, VectorRegister4f c)                                        { return VectorCastIntToFloat(VectorIntAnd(VectorCastFloatToInt(VectorCompareGE(a, b)), *(VectorRegister4i *)&c)); }
-VM_FORCEINLINE VectorRegister4f VVM_Exec3f_cmpeq_logic_and             (FVectorVMBatchState *BatchState, VectorRegister4f a, VectorRegister4f b, VectorRegister4f c)                                        { return VectorCastIntToFloat(VectorIntAnd(VectorCastFloatToInt(VectorCompareEQ(a, b)), *(VectorRegister4i *)&c)); }
-VM_FORCEINLINE VectorRegister4f VVM_Exec3f_cmpne_logic_and             (FVectorVMBatchState *BatchState, VectorRegister4f a, VectorRegister4f b, VectorRegister4f c)                                        { return VectorCastIntToFloat(VectorIntAnd(VectorCastFloatToInt(VectorCompareNE(a, b)), *(VectorRegister4i *)&c)); }
-VM_FORCEINLINE VectorRegister4i VVM_Exec3i_cmplti_logic_and            (FVectorVMBatchState *BatchState, VectorRegister4i a, VectorRegister4i b, VectorRegister4i c)                                        { return VectorIntAnd(VectorIntCompareLT(a, b), c); }
-VM_FORCEINLINE VectorRegister4i VVM_Exec3i_cmplei_logic_and            (FVectorVMBatchState *BatchState, VectorRegister4i a, VectorRegister4i b, VectorRegister4i c)                                        { return VectorIntAnd(VectorIntCompareLE(a, b), c); }
-VM_FORCEINLINE VectorRegister4i VVM_Exec3i_cmpgti_logic_and            (FVectorVMBatchState *BatchState, VectorRegister4i a, VectorRegister4i b, VectorRegister4i c)                                        { return VectorIntAnd(VectorIntCompareGT(a, b), c); }
-VM_FORCEINLINE VectorRegister4i VVM_Exec3i_cmpgei_logic_and            (FVectorVMBatchState *BatchState, VectorRegister4i a, VectorRegister4i b, VectorRegister4i c)                                        { return VectorIntAnd(VectorIntCompareGE(a, b), c); }
-VM_FORCEINLINE VectorRegister4i VVM_Exec3i_cmpeqi_logic_and            (FVectorVMBatchState *BatchState, VectorRegister4i a, VectorRegister4i b, VectorRegister4i c)                                        { return VectorIntAnd(VectorIntCompareEQ(a, b), c); }
-VM_FORCEINLINE VectorRegister4i VVM_Exec3i_cmpnei_logic_and            (FVectorVMBatchState *BatchState, VectorRegister4i a, VectorRegister4i b, VectorRegister4i c)                                        { return VectorIntAnd(VectorIntCompareNEQ(a, b), c); }
-VM_FORCEINLINE VectorRegister4f VVM_Exec3f_cmplt_logic_or              (FVectorVMBatchState *BatchState, VectorRegister4f a, VectorRegister4f b, VectorRegister4f c)                                        { return VectorCastIntToFloat(VectorIntOr(VectorCastFloatToInt(VectorCompareLT(a, b)), *(VectorRegister4i *)&c)); }
-VM_FORCEINLINE VectorRegister4f VVM_Exec3f_cmple_logic_or              (FVectorVMBatchState *BatchState, VectorRegister4f a, VectorRegister4f b, VectorRegister4f c)                                        { return VectorCastIntToFloat(VectorIntOr(VectorCastFloatToInt(VectorCompareLE(a, b)), *(VectorRegister4i *)&c)); }
-VM_FORCEINLINE VectorRegister4f VVM_Exec3f_cmpgt_logic_or              (FVectorVMBatchState *BatchState, VectorRegister4f a, VectorRegister4f b, VectorRegister4f c)                                        { return VectorCastIntToFloat(VectorIntOr(VectorCastFloatToInt(VectorCompareGT(a, b)), *(VectorRegister4i *)&c)); }
-VM_FORCEINLINE VectorRegister4f VVM_Exec3f_cmpge_logic_or              (FVectorVMBatchState *BatchState, VectorRegister4f a, VectorRegister4f b, VectorRegister4f c)                                        { return VectorCastIntToFloat(VectorIntOr(VectorCastFloatToInt(VectorCompareGE(a, b)), *(VectorRegister4i *)&c)); }
-VM_FORCEINLINE VectorRegister4f VVM_Exec3f_cmpeq_logic_or              (FVectorVMBatchState *BatchState, VectorRegister4f a, VectorRegister4f b, VectorRegister4f c)                                        { return VectorCastIntToFloat(VectorIntOr(VectorCastFloatToInt(VectorCompareEQ(a, b)), *(VectorRegister4i *)&c)); }
-VM_FORCEINLINE VectorRegister4f VVM_Exec3f_cmpne_logic_or              (FVectorVMBatchState *BatchState, VectorRegister4f a, VectorRegister4f b, VectorRegister4f c)                                        { return VectorCastIntToFloat(VectorIntOr(VectorCastFloatToInt(VectorCompareNE(a, b)), *(VectorRegister4i *)&c)); }
-VM_FORCEINLINE VectorRegister4i VVM_Exec3i_cmplti_logic_or             (FVectorVMBatchState *BatchState, VectorRegister4i a, VectorRegister4i b, VectorRegister4i c)                                        { return VectorIntOr(VectorIntCompareLT(a, b), c); }
-VM_FORCEINLINE VectorRegister4i VVM_Exec3i_cmplei_logic_or             (FVectorVMBatchState *BatchState, VectorRegister4i a, VectorRegister4i b, VectorRegister4i c)                                        { return VectorIntOr(VectorIntCompareLE(a, b), c); }
-VM_FORCEINLINE VectorRegister4i VVM_Exec3i_cmpgti_logic_or             (FVectorVMBatchState *BatchState, VectorRegister4i a, VectorRegister4i b, VectorRegister4i c)                                        { return VectorIntOr(VectorIntCompareGT(a, b), c); }
-VM_FORCEINLINE VectorRegister4i VVM_Exec3i_cmpgei_logic_or             (FVectorVMBatchState *BatchState, VectorRegister4i a, VectorRegister4i b, VectorRegister4i c)                                        { return VectorIntOr(VectorIntCompareGE(a, b), c); }
-VM_FORCEINLINE VectorRegister4i VVM_Exec3i_cmpeqi_logic_or             (FVectorVMBatchState *BatchState, VectorRegister4i a, VectorRegister4i b, VectorRegister4i c)                                        { return VectorIntOr(VectorIntCompareEQ(a, b), c); }
-VM_FORCEINLINE VectorRegister4i VVM_Exec3i_cmpnei_logic_or             (FVectorVMBatchState *BatchState, VectorRegister4i a, VectorRegister4i b, VectorRegister4i c)                                        { return VectorIntOr(VectorIntCompareNEQ(a, b), c); }
-VM_FORCEINLINE VectorRegister4f VVM_Exec4f_mad_add                     (FVectorVMBatchState *BatchState, VectorRegister4f a, VectorRegister4f b, VectorRegister4f c, VectorRegister4f d)                    { return VectorAdd(VectorMultiplyAdd(a, b, c), d); }
-VM_FORCEINLINE VectorRegister4f VVM_Exec4f_mad_sub0                    (FVectorVMBatchState *BatchState, VectorRegister4f a, VectorRegister4f b, VectorRegister4f c, VectorRegister4f d)                    { return VectorSubtract(VectorMultiplyAdd(a, b, c), d); }
-VM_FORCEINLINE VectorRegister4f VVM_Exec4f_mad_sub1                    (FVectorVMBatchState *BatchState, VectorRegister4f a, VectorRegister4f b, VectorRegister4f c, VectorRegister4f d)                    { return VectorSubtract(d, VectorMultiplyAdd(a, b, c)); }
-VM_FORCEINLINE VectorRegister4f VVM_Exec4f_mad_mul                     (FVectorVMBatchState *BatchState, VectorRegister4f a, VectorRegister4f b, VectorRegister4f c, VectorRegister4f d)                    { return VectorMultiply(VectorMultiplyAdd(a, b, c), d); }
-VM_FORCEINLINE VectorRegister4f VVM_Exec3f_mad_sqrt                    (FVectorVMBatchState *BatchState, VectorRegister4f a, VectorRegister4f b, VectorRegister4f c)                                        { return VectorSqrt(VectorMultiplyAdd(a, b, c)); }
-VM_FORCEINLINE VectorRegister4f VVM_Exec5f_mad_mad0                    (FVectorVMBatchState *BatchState, VectorRegister4f a, VectorRegister4f b, VectorRegister4f c, VectorRegister4f d, VectorRegister4f e){ return VectorMultiplyAdd(d, e, VectorMultiplyAdd(a, b, c)); }
-VM_FORCEINLINE VectorRegister4f VVM_Exec5f_mad_mad1                    (FVectorVMBatchState *BatchState, VectorRegister4f a, VectorRegister4f b, VectorRegister4f c, VectorRegister4f d, VectorRegister4f e){ return VectorMultiplyAdd(VectorMultiplyAdd(a, b, c), d, e); }
-VM_FORCEINLINE VectorRegister4f VVM_Exec4f_mul_mad0                    (FVectorVMBatchState *BatchState, VectorRegister4f a, VectorRegister4f b, VectorRegister4f c, VectorRegister4f d)                    { return VectorMultiplyAdd(VectorMultiply(a, b), c, d); }
-VM_FORCEINLINE VectorRegister4f VVM_Exec4f_mul_mad1                    (FVectorVMBatchState *BatchState, VectorRegister4f a, VectorRegister4f b, VectorRegister4f c, VectorRegister4f d)                    { return VectorMultiplyAdd(c, d, VectorMultiply(a, b)); }
-VM_FORCEINLINE VectorRegister4f VVM_Exec3f_mul_add                     (FVectorVMBatchState *BatchState, VectorRegister4f a, VectorRegister4f b, VectorRegister4f c)                                        { return VectorAdd(VectorMultiply(a, b), c); }
-VM_FORCEINLINE VectorRegister4f VVM_Exec3f_mul_sub0                    (FVectorVMBatchState *BatchState, VectorRegister4f a, VectorRegister4f b, VectorRegister4f c)                                        { return VectorSubtract(VectorMultiply(a, b), c); }
-VM_FORCEINLINE VectorRegister4f VVM_Exec3f_mul_sub1                    (FVectorVMBatchState *BatchState, VectorRegister4f a, VectorRegister4f b, VectorRegister4f c)                                        { return VectorSubtract(c, VectorMultiply(a, b)); }
-VM_FORCEINLINE VectorRegister4f VVM_Exec3f_mul_mul                     (FVectorVMBatchState *BatchState, VectorRegister4f a, VectorRegister4f b, VectorRegister4f c)                                        { return VectorMultiply(VectorMultiply(a, b), c); }
-VM_FORCEINLINE VectorRegister4f VVM_Exec3f_mul_max                     (FVectorVMBatchState *BatchState, VectorRegister4f a, VectorRegister4f b, VectorRegister4f c)                                        { return VectorMax(VectorMultiply(a, b), c); }
-VM_FORCEINLINE VectorRegister4f VVM_Exec2f_mul_2x                      (FVectorVMBatchState *BatchState, VectorRegister4f a, VectorRegister4f b)                                                            { return VectorMultiply(a, b); }
-VM_FORCEINLINE VectorRegister4f VVM_Exec4f_add_mad1                    (FVectorVMBatchState *BatchState, VectorRegister4f a, VectorRegister4f b, VectorRegister4f c, VectorRegister4f d)                    { return VectorMultiplyAdd(c, d, VectorAdd(a, b)); }
-VM_FORCEINLINE VectorRegister4f VVM_Exec3f_add_add                     (FVectorVMBatchState *BatchState, VectorRegister4f a, VectorRegister4f b, VectorRegister4f c)                                        { return VectorAdd(VectorAdd(a, b), c); }
-VM_FORCEINLINE VectorRegister4f VVM_Exec3f_sub_cmplt1                  (FVectorVMBatchState *BatchState, VectorRegister4f a, VectorRegister4f b, VectorRegister4f c)                                        { return VectorCompareLT(c, VectorSubtract(a, b)); }
-VM_FORCEINLINE VectorRegister4f VVM_Exec2f_sub_neg                     (FVectorVMBatchState *BatchState, VectorRegister4f a, VectorRegister4f b)                                                            { return VectorNegate(VectorSubtract(a, b)); }
-VM_FORCEINLINE VectorRegister4f VVM_Exec3f_sub_mul                     (FVectorVMBatchState *BatchState, VectorRegister4f a, VectorRegister4f b, VectorRegister4f c)                                        { return VectorMultiply(VectorSubtract(a, b), c); }
-VM_FORCEINLINE VectorRegister4f VVM_Exec4f_div_mad0                    (FVectorVMBatchState *BatchState, VectorRegister4f a, VectorRegister4f b, VectorRegister4f c, VectorRegister4f d)                    { return VectorMultiplyAdd(VVM_Exec2f_div(BatchState, a, b), c, d); }
-VM_FORCEINLINE VectorRegister4i VVM_Exec2i_div_f2i                     (FVectorVMBatchState *BatchState, VectorRegister4i a, VectorRegister4i b)                                                            { return VectorFloatToInt(VVM_Exec2f_div(BatchState, VectorCastIntToFloat(a), VectorCastIntToFloat(b))); }
-VM_FORCEINLINE VectorRegister4f VVM_Exec3f_div_mul                     (FVectorVMBatchState *BatchState, VectorRegister4f a, VectorRegister4f b, VectorRegister4f c)                                        { return VectorMultiply(VVM_Exec2f_div(BatchState, a, b), c); }
-VM_FORCEINLINE VectorRegister4i VVM_Exec3i_muli_addi                   (FVectorVMBatchState *BatchState, VectorRegister4i a, VectorRegister4i b, VectorRegister4i c)                                        { return VectorIntAdd(VectorIntMultiply(a, b), c); }
-VM_FORCEINLINE VectorRegister4i VVM_Exec3i_addi_bit_rshift             (FVectorVMBatchState *BatchState, VectorRegister4i a, VectorRegister4i b, VectorRegister4i c)                                        { return VVMIntRShift(VectorIntAdd(a, b), c); }
-VM_FORCEINLINE VectorRegister4i VVM_Exec3i_addi_muli                   (FVectorVMBatchState *BatchState, VectorRegister4i a, VectorRegister4i b, VectorRegister4i c)                                        { return VectorIntMultiply(VectorIntAdd(a, b), c); }
-VM_FORCEINLINE VectorRegister4i VVM_Exec1i_b2i_2x                      (FVectorVMBatchState *BatchState, VectorRegister4i a)                                                                                { return VectorIntSelect(a, VectorIntSet1(1), VectorSetZero()); }
-VM_FORCEINLINE VectorRegister4f VVM_Exec2f_i2f_div0                    (FVectorVMBatchState *BatchState, VectorRegister4f a, VectorRegister4f b)                                                            { return VVM_Exec2f_div(BatchState, VectorIntToFloat(VectorCastFloatToInt(a)), b); }
-VM_FORCEINLINE VectorRegister4f VVM_Exec2f_i2f_div1                    (FVectorVMBatchState *BatchState, VectorRegister4f a, VectorRegister4f b)                                                            { return VVM_Exec2f_div(BatchState, b, VectorIntToFloat(VectorCastFloatToInt(a))); }
-VM_FORCEINLINE VectorRegister4f VVM_Exec2f_i2f_mul                     (FVectorVMBatchState *BatchState, VectorRegister4f a, VectorRegister4f b)                                                            { return VectorMultiply(VectorIntToFloat(VectorCastFloatToInt(a)), b); }
-VM_FORCEINLINE VectorRegister4f VVM_Exec3f_i2f_mad0                    (FVectorVMBatchState *BatchState, VectorRegister4f a, VectorRegister4f b, VectorRegister4f c)                                        { return VectorMultiplyAdd(VectorIntToFloat(VectorCastFloatToInt(a)), b, c); }
-VM_FORCEINLINE VectorRegister4f VVM_Exec3f_i2f_mad1                    (FVectorVMBatchState *BatchState, VectorRegister4f a, VectorRegister4f b, VectorRegister4f c)                                        { return VectorMultiplyAdd(a, b, VectorIntToFloat(VectorCastFloatToInt(c))); }
-VM_FORCEINLINE VectorRegister4i VVM_Exec3i_f2i_select1                 (FVectorVMBatchState *BatchState, VectorRegister4i a, VectorRegister4i b, VectorRegister4i c)                                        { return VectorIntSelect(a, VectorFloatToInt(VectorCastIntToFloat(b)), c); }
-VM_FORCEINLINE VectorRegister4i VVM_Exec2i_f2i_maxi                    (FVectorVMBatchState *BatchState, VectorRegister4i a, VectorRegister4i b)                                                            { return VectorIntMax(VectorFloatToInt(VectorCastIntToFloat(a)), b); }
-VM_FORCEINLINE VectorRegister4i VVM_Exec2i_f2i_addi                    (FVectorVMBatchState *BatchState, VectorRegister4i a, VectorRegister4i b)                                                            { return VectorIntAdd(VectorFloatToInt(VectorCastIntToFloat(a)), b); }
-VM_FORCEINLINE VectorRegister4f VVM_Exec3f_fmod_add                    (FVectorVMBatchState *BatchState, VectorRegister4f a, VectorRegister4f b, VectorRegister4f c)                                        { return VectorAdd(VectorMod(a, b), c); }
-VM_FORCEINLINE VectorRegister4f VVM_Exec2f_bit_and_i2f                 (FVectorVMBatchState *BatchState, VectorRegister4f a, VectorRegister4f b)                                                            { return VectorIntToFloat(VectorIntAnd(VectorCastFloatToInt(a), VectorCastFloatToInt(b))); }
-VM_FORCEINLINE VectorRegister4i VVM_Exec3i_bit_rshift_bit_and          (FVectorVMBatchState *BatchState, VectorRegister4i a, VectorRegister4i b, VectorRegister4i c)                                        { return VectorIntAnd(VVMIntRShift(a, b), c); }
-VM_FORCEINLINE VectorRegister4f VVM_Exec2f_neg_cmplt                   (FVectorVMBatchState *BatchState, VectorRegister4f a, VectorRegister4f b)                                                            { return VectorCompareLT(VectorNegate(a), b); }
-VM_FORCEINLINE VectorRegister4i VVM_Exec3i_bit_or_muli                 (FVectorVMBatchState *BatchState, VectorRegister4i a, VectorRegister4i b, VectorRegister4i c)                                        { return VectorIntMultiply(VectorIntOr(a, b), c); }
-VM_FORCEINLINE VectorRegister4i VVM_Exec3i_bit_lshift_bit_or           (FVectorVMBatchState *BatchState, VectorRegister4i a, VectorRegister4i b, VectorRegister4i c)                                        { return VectorIntOr(VVMIntLShift(a, b), c); }
-VM_FORCEINLINE const uint8 *    VVM_Exec2null_random_add               (VVM_NULL_FN_ARGS)                                                                                                                   { return VVM_random_add(CT_MultipleLoops, InsPtr, BatchState, ExecCtx, SerializeState, CmpSerializeState, NumLoops); }
-VM_FORCEINLINE const uint8 *    VVM_Exec1null_random_2x                (VVM_NULL_FN_ARGS)                                                                                                                   { return VVM_random_2x(CT_MultipleLoops, InsPtr, BatchState, ExecCtx, SerializeState, CmpSerializeState, NumLoops); }
-VM_FORCEINLINE VectorRegister4i VVM_Exec2i_max_f2i                     (FVectorVMBatchState *BatchState, VectorRegister4i a, VectorRegister4i b)                                                            { return VectorFloatToInt(VectorMax(VectorCastIntToFloat(a), VectorCastIntToFloat(b))); }
-VM_FORCEINLINE VectorRegister4f VVM_Exec4f_select_mul                  (FVectorVMBatchState *BatchState, VectorRegister4f a, VectorRegister4f b, VectorRegister4f c, VectorRegister4f d)                    { return VectorMultiply(VectorSelect(a, b, c), d); }
-VM_FORCEINLINE VectorRegister4f VVM_Exec4f_select_add                  (FVectorVMBatchState *BatchState, VectorRegister4f a, VectorRegister4f b, VectorRegister4f c, VectorRegister4f d)                    { return VectorAdd(VectorSelect(a, b, c), d); }
-VM_FORCEINLINE const uint8 *    VVM_Exec1null_sin_cos                  (VVM_NULL_FN_ARGS)                                                                                                                   { return VVM_sincos(CT_MultipleLoops, InsPtr, BatchState, NumLoops); }
-VM_FORCEINLINE const uint8 *    VVM_Exec0null_inputdata_float          (VVM_NULL_FN_ARGS)                                                                                                                   { return InsPtr; }
-VM_FORCEINLINE const uint8 *    VVM_Exec0null_inputdata_int32          (VVM_NULL_FN_ARGS)                                                                                                                   { return InsPtr; }
-VM_FORCEINLINE const uint8 *    VVM_Exec0null_inputdata_half           (VVM_NULL_FN_ARGS)                                                                                                                   { return InsPtr; }
-VM_FORCEINLINE const uint8 *    VVM_Exec0null_inputdata_noadvance_float(VVM_NULL_FN_ARGS)                                                                                                                   { return InsPtr; }
-VM_FORCEINLINE const uint8 *    VVM_Exec0null_inputdata_noadvance_int32(VVM_NULL_FN_ARGS)                                                                                                                   { return InsPtr; }
-VM_FORCEINLINE const uint8 *    VVM_Exec0null_inputdata_noadvance_half (VVM_NULL_FN_ARGS)                                                                                                                   { return InsPtr; }
+VM_FORCEINLINE VectorRegister4f VVM_Exec2f_add                             (FVectorVMBatchState *BatchState, VectorRegister4f a, VectorRegister4f b)                                                            { return VectorAdd(a, b); }
+VM_FORCEINLINE VectorRegister4f VVM_Exec2f_sub                             (FVectorVMBatchState *BatchState, VectorRegister4f a, VectorRegister4f b)                                                            { return VectorSubtract(a, b); }
+VM_FORCEINLINE VectorRegister4f VVM_Exec2f_mul                             (FVectorVMBatchState *BatchState, VectorRegister4f a, VectorRegister4f b)                                                            { return VectorMultiply(a, b); }
+VM_FORCEINLINE VectorRegister4f VVM_Exec2f_div                             (FVectorVMBatchState *BatchState, VectorRegister4f a, VectorRegister4f b)                                                            { return VectorSelect(VectorCompareGT(VectorAbs(b), VVM_m128Const(Epsilon)), VectorDivide(a, b), VectorZeroFloat()); }
+VM_FORCEINLINE VectorRegister4f VVM_Exec3f_mad                             (FVectorVMBatchState *BatchState, VectorRegister4f a, VectorRegister4f b, VectorRegister4f c)                                        { return VectorMultiplyAdd(a, b, c); }
+VM_FORCEINLINE VectorRegister4f VVM_Exec3f_lerp                            (FVectorVMBatchState *BatchState, VectorRegister4f a, VectorRegister4f b, VectorRegister4f c)                                        { return VectorLerp(a, b, c); }
+VM_FORCEINLINE VectorRegister4f VVM_Exec1f_rcp                             (FVectorVMBatchState *BatchState, VectorRegister4f a)                                                                                { return VectorSelect(VectorCompareGT(VectorAbs(a) , VVM_m128Const(Epsilon)), VectorReciprocalEstimate(a), VectorZeroFloat()); }
+VM_FORCEINLINE VectorRegister4f VVM_Exec1f_rsq                             (FVectorVMBatchState *BatchState, VectorRegister4f a)                                                                                { return VectorSelect(VectorCompareGT(a, VVM_m128Const(Epsilon)), VectorReciprocalSqrt(a), VectorZeroFloat()); }
+VM_FORCEINLINE VectorRegister4f VVM_Exec1f_sqrt                            (FVectorVMBatchState *BatchState, VectorRegister4f a)                                                                                { return VectorSelect(VectorCompareGT(a, VVM_m128Const(Epsilon)), VectorSqrt(a), VectorZeroFloat()); }
+VM_FORCEINLINE VectorRegister4f VVM_Exec1f_neg                             (FVectorVMBatchState *BatchState, VectorRegister4f a)                                                                                { return VectorNegate(a); }
+VM_FORCEINLINE VectorRegister4f VVM_Exec1f_abs                             (FVectorVMBatchState *BatchState, VectorRegister4f a)                                                                                { return VectorAbs(a); }
+VM_FORCEINLINE VectorRegister4f VVM_Exec1f_exp                             (FVectorVMBatchState *BatchState, VectorRegister4f a)                                                                                { return VectorExp(a); }
+VM_FORCEINLINE VectorRegister4f VVM_Exec1f_exp2                            (FVectorVMBatchState *BatchState, VectorRegister4f a)                                                                                { return VectorExp2(a); }
+VM_FORCEINLINE VectorRegister4f VVM_Exec1f_log                             (FVectorVMBatchState *BatchState, VectorRegister4f a)                                                                                { return VectorSelect(VectorCompareGT(a, VectorZeroFloat()), VectorLog(a), VectorZeroFloat()); }
+VM_FORCEINLINE VectorRegister4f VVM_Exec1f_log2                            (FVectorVMBatchState *BatchState, VectorRegister4f a)                                                                                { return VectorLog2(a); }
+VM_FORCEINLINE VectorRegister4f VVM_Exec1f_sin                             (FVectorVMBatchState *BatchState, VectorRegister4f a)                                                                                { return VectorSin(a); }
+VM_FORCEINLINE VectorRegister4f VVM_Exec1f_cos                             (FVectorVMBatchState *BatchState, VectorRegister4f a)                                                                                { return VectorCos(a); }
+VM_FORCEINLINE VectorRegister4f VVM_Exec1f_tan                             (FVectorVMBatchState *BatchState, VectorRegister4f a)                                                                                { return VectorTan(a); }
+VM_FORCEINLINE VectorRegister4f VVM_Exec1f_acos                            (FVectorVMBatchState *BatchState, VectorRegister4f a)                                                                                { return VectorATan2(VVM_Exec1f_sqrt(BatchState, VectorMultiply(VectorSubtract(VVM_m128Const(One), a), VectorAdd(VVM_m128Const(One), a))), a); }
+VM_FORCEINLINE VectorRegister4f VVM_Exec1f_asin                            (FVectorVMBatchState *BatchState, VectorRegister4f a)                                                                                { return VectorSubtract(VVM_m128Const(QuarterPi), VVM_Exec1f_acos(BatchState, a)); }
+VM_FORCEINLINE VectorRegister4f VVM_Exec1f_atan                            (FVectorVMBatchState *BatchState, VectorRegister4f a)                                                                                { return VectorATan(a); }
+VM_FORCEINLINE VectorRegister4f VVM_Exec2f_atan2                           (FVectorVMBatchState *BatchState, VectorRegister4f a, VectorRegister4f b)                                                            { return VectorATan2(a, b); }
+VM_FORCEINLINE VectorRegister4f VVM_Exec1f_ceil                            (FVectorVMBatchState *BatchState, VectorRegister4f a)                                                                                { return VectorCeil(a); }
+VM_FORCEINLINE VectorRegister4f VVM_Exec1f_floor                           (FVectorVMBatchState *BatchState, VectorRegister4f a)                                                                                { return VectorFloor(a); }
+VM_FORCEINLINE VectorRegister4f VVM_Exec2f_fmod                            (FVectorVMBatchState *BatchState, VectorRegister4f a, VectorRegister4f b)                                                            { return VectorMod(a, b); }
+VM_FORCEINLINE VectorRegister4f VVM_Exec1f_frac                            (FVectorVMBatchState *BatchState, VectorRegister4f a)                                                                                { return VectorFractional(a); }
+VM_FORCEINLINE VectorRegister4f VVM_Exec1f_trunc                           (FVectorVMBatchState *BatchState, VectorRegister4f a)                                                                                { return VectorTruncate(a); }
+VM_FORCEINLINE VectorRegister4f VVM_Exec3f_clamp                           (FVectorVMBatchState *BatchState, VectorRegister4f a, VectorRegister4f b, VectorRegister4f c)                                        { return VectorClamp(a, b, c); }
+VM_FORCEINLINE VectorRegister4f VVM_Exec2f_min                             (FVectorVMBatchState *BatchState, VectorRegister4f a, VectorRegister4f b)                                                            { return VectorMin(a, b); }
+VM_FORCEINLINE VectorRegister4f VVM_Exec2f_max                             (FVectorVMBatchState *BatchState, VectorRegister4f a, VectorRegister4f b)                                                            { return VectorMax(a, b); }
+VM_FORCEINLINE VectorRegister4f VVM_Exec2f_pow                             (FVectorVMBatchState *BatchState, VectorRegister4f a, VectorRegister4f b)                                                            { return VectorSelect(VectorCompareGT(a, VVM_m128Const(Epsilon)), VectorPow(a, b), VectorZeroFloat()); }
+VM_FORCEINLINE VectorRegister4f VVM_Exec1f_round                           (FVectorVMBatchState *BatchState, VectorRegister4f a)                                                                                { return VectorRound(a); }
+VM_FORCEINLINE VectorRegister4f VVM_Exec1f_sign                            (FVectorVMBatchState *BatchState, VectorRegister4f a)                                                                                { return VectorSign(a); }
+VM_FORCEINLINE VectorRegister4f VVM_Exec2f_step                            (FVectorVMBatchState *BatchState, VectorRegister4f a, VectorRegister4f b)                                                            { return VectorStep(VectorSubtract(a, b)); }
+VM_FORCEINLINE const uint8 *    VVM_Exec1null_random                       (VVM_NULL_FN_ARGS)                                                                                                                   { return VVM_random(CT_MultipleLoops, InsPtr, BatchState, ExecCtx, SerializeState, CmpSerializeState, NumLoops); }
+VM_FORCEINLINE const uint8 *    VVM_Exec0null_noise                        (VVM_NULL_FN_ARGS)                                                                                                                   { return InsPtr; }
+VM_FORCEINLINE VectorRegister4f VVM_Exec2f_cmplt                           (FVectorVMBatchState *BatchState, VectorRegister4f a, VectorRegister4f b)                                                            { return VectorCompareLT(a, b); }
+VM_FORCEINLINE VectorRegister4f VVM_Exec2f_cmple                           (FVectorVMBatchState *BatchState, VectorRegister4f a, VectorRegister4f b)                                                            { return VectorCompareLE(a, b); }
+VM_FORCEINLINE VectorRegister4f VVM_Exec2f_cmpgt                           (FVectorVMBatchState *BatchState, VectorRegister4f a, VectorRegister4f b)                                                            { return VectorCompareGT(a, b); }
+VM_FORCEINLINE VectorRegister4f VVM_Exec2f_cmpge                           (FVectorVMBatchState *BatchState, VectorRegister4f a, VectorRegister4f b)                                                            { return VectorCompareGE(a, b); }
+VM_FORCEINLINE VectorRegister4f VVM_Exec2f_cmpeq                           (FVectorVMBatchState *BatchState, VectorRegister4f a, VectorRegister4f b)                                                            { return VectorCompareEQ(a, b); }
+VM_FORCEINLINE VectorRegister4f VVM_Exec2f_cmpneq                          (FVectorVMBatchState *BatchState, VectorRegister4f a, VectorRegister4f b)                                                            { return VectorCompareNE(a, b); }
+VM_FORCEINLINE VectorRegister4f VVM_Exec3f_select                          (FVectorVMBatchState *BatchState, VectorRegister4f a, VectorRegister4f b, VectorRegister4f c)                                        { return VectorSelect(a, b, c); }
+VM_FORCEINLINE VectorRegister4i VVM_Exec2i_addi                            (FVectorVMBatchState *BatchState, VectorRegister4i a, VectorRegister4i b)                                                            { return VectorIntAdd(a, b); }
+VM_FORCEINLINE VectorRegister4i VVM_Exec2i_subi                            (FVectorVMBatchState *BatchState, VectorRegister4i a, VectorRegister4i b)                                                            { return VectorIntSubtract(a, b); }
+VM_FORCEINLINE VectorRegister4i VVM_Exec2i_muli                            (FVectorVMBatchState *BatchState, VectorRegister4i a, VectorRegister4i b)                                                            { return VectorIntMultiply(a, b); }
+VM_FORCEINLINE VectorRegister4i VVM_Exec2i_divi                            (FVectorVMBatchState *BatchState, VectorRegister4i a, VectorRegister4i b)                                                            { return VVMIntDiv(a, b); }
+VM_FORCEINLINE VectorRegister4i VVM_Exec3i_clampi                          (FVectorVMBatchState *BatchState, VectorRegister4i a, VectorRegister4i b, VectorRegister4i c)                                        { return VectorIntClamp(a, b, c); }
+VM_FORCEINLINE VectorRegister4i VVM_Exec2i_mini                            (FVectorVMBatchState *BatchState, VectorRegister4i a, VectorRegister4i b)                                                            { return VectorIntMin(a, b); }
+VM_FORCEINLINE VectorRegister4i VVM_Exec2i_maxi                            (FVectorVMBatchState *BatchState, VectorRegister4i a, VectorRegister4i b)                                                            { return VectorIntMax(a, b); }
+VM_FORCEINLINE VectorRegister4i VVM_Exec1i_absi                            (FVectorVMBatchState *BatchState, VectorRegister4i a)                                                                                { return VectorIntAbs(a); }
+VM_FORCEINLINE VectorRegister4i VVM_Exec1i_negi                            (FVectorVMBatchState *BatchState, VectorRegister4i a)                                                                                { return VectorIntNegate(a); }
+VM_FORCEINLINE VectorRegister4i VVM_Exec1i_signi                           (FVectorVMBatchState *BatchState, VectorRegister4i a)                                                                                { return VectorIntSign(a); }
+VM_FORCEINLINE const uint8 *    VVM_Exec1null_randomi                      (VVM_NULL_FN_ARGS)                                                                                                                   { return VVM_randomi(CT_MultipleLoops, InsPtr, BatchState, ExecCtx, SerializeState, CmpSerializeState, NumLoops); }
+VM_FORCEINLINE VectorRegister4i VVM_Exec2i_cmplti                          (FVectorVMBatchState *BatchState, VectorRegister4i a, VectorRegister4i b)                                                            { return VectorIntCompareLT(a, b); }
+VM_FORCEINLINE VectorRegister4i VVM_Exec2i_cmplei                          (FVectorVMBatchState *BatchState, VectorRegister4i a, VectorRegister4i b)                                                            { return VectorIntCompareLE(a, b); }
+VM_FORCEINLINE VectorRegister4i VVM_Exec2i_cmpgti                          (FVectorVMBatchState *BatchState, VectorRegister4i a, VectorRegister4i b)                                                            { return VectorIntCompareGT(a, b); }
+VM_FORCEINLINE VectorRegister4i VVM_Exec2i_cmpgei                          (FVectorVMBatchState *BatchState, VectorRegister4i a, VectorRegister4i b)                                                            { return VectorIntCompareGE(a, b); }
+VM_FORCEINLINE VectorRegister4i VVM_Exec2i_cmpeqi                          (FVectorVMBatchState *BatchState, VectorRegister4i a, VectorRegister4i b)                                                            { return VectorIntCompareEQ(a, b); }
+VM_FORCEINLINE VectorRegister4i VVM_Exec2i_cmpneqi                         (FVectorVMBatchState *BatchState, VectorRegister4i a, VectorRegister4i b)                                                            { return VectorIntCompareNEQ(a, b); }
+VM_FORCEINLINE VectorRegister4i VVM_Exec2i_bit_and                         (FVectorVMBatchState *BatchState, VectorRegister4i a, VectorRegister4i b)                                                            { return VectorIntAnd(a, b); }
+VM_FORCEINLINE VectorRegister4i VVM_Exec2i_bit_or                          (FVectorVMBatchState *BatchState, VectorRegister4i a, VectorRegister4i b)                                                            { return VectorIntOr(a, b); }
+VM_FORCEINLINE VectorRegister4i VVM_Exec2i_bit_xor                         (FVectorVMBatchState *BatchState, VectorRegister4i a, VectorRegister4i b)                                                            { return VectorIntXor(a, b); }
+VM_FORCEINLINE VectorRegister4i VVM_Exec1i_bit_not                         (FVectorVMBatchState *BatchState, VectorRegister4i a)                                                                                { return VectorIntNot(a); }
+VM_FORCEINLINE VectorRegister4i VVM_Exec2i_bit_lshift                      (FVectorVMBatchState *BatchState, VectorRegister4i a, VectorRegister4i b)                                                            { return VVMIntLShift(a, b); }
+VM_FORCEINLINE VectorRegister4i VVM_Exec2i_bit_rshift                      (FVectorVMBatchState *BatchState, VectorRegister4i a, VectorRegister4i b)                                                            { return VVMIntRShift(a, b); }
+VM_FORCEINLINE VectorRegister4i VVM_Exec2i_logic_and                       (FVectorVMBatchState *BatchState, VectorRegister4i a, VectorRegister4i b)                                                            { return VectorIntAnd(a, b); }
+VM_FORCEINLINE VectorRegister4i VVM_Exec2i_logic_or                        (FVectorVMBatchState *BatchState, VectorRegister4i a, VectorRegister4i b)                                                            { return VectorIntOr(a, b); }
+VM_FORCEINLINE VectorRegister4i VVM_Exec2i_logic_xor                       (FVectorVMBatchState *BatchState, VectorRegister4i a, VectorRegister4i b)                                                            { return VectorIntXor(a, b); }
+VM_FORCEINLINE VectorRegister4i VVM_Exec1i_logic_not                       (FVectorVMBatchState *BatchState, VectorRegister4i a)                                                                                { return VectorIntNot(a); }
+VM_FORCEINLINE VectorRegister4i VVM_Exec1i_f2i                             (FVectorVMBatchState *BatchState, VectorRegister4i a)                                                                                { return VVMf2i(a); }
+VM_FORCEINLINE VectorRegister4f VVM_Exec1f_i2f                             (FVectorVMBatchState *BatchState, VectorRegister4f a)                                                                                { return VVMi2f(a); }
+VM_FORCEINLINE VectorRegister4f VVM_Exec1f_f2b                             (FVectorVMBatchState *BatchState, VectorRegister4f a)                                                                                { return VectorCompareGT(a, VectorZeroFloat()); }
+VM_FORCEINLINE VectorRegister4f VVM_Exec1f_b2f                             (FVectorVMBatchState *BatchState, VectorRegister4f a)                                                                                { return VectorSelect(a, VVM_m128Const(One), VectorZeroFloat()); }
+VM_FORCEINLINE VectorRegister4i VVM_Exec1i_i2b                             (FVectorVMBatchState *BatchState, VectorRegister4i a)                                                                                { return VectorIntCompareGT(a, VectorSetZero()); }
+VM_FORCEINLINE VectorRegister4i VVM_Exec1i_b2i                             (FVectorVMBatchState *BatchState, VectorRegister4i a)                                                                                { return VectorIntSelect(a, VectorIntSet1(1), VectorSetZero()); }
+VM_FORCEINLINE const uint8 *    VVM_Exec0null_outputdata_float             (VVM_NULL_FN_ARGS)                                                                                                                   { return VVM_Output32(CT_MultipleLoops, InsPtr, BatchState, ExecCtx); }
+VM_FORCEINLINE const uint8 *    VVM_Exec0null_outputdata_int32             (VVM_NULL_FN_ARGS)                                                                                                                   { return VVM_Output32(CT_MultipleLoops, InsPtr, BatchState, ExecCtx); }
+VM_FORCEINLINE const uint8 *    VVM_Exec0null_outputdata_half              (VVM_NULL_FN_ARGS)                                                                                                                   { return VVM_Output16(CT_MultipleLoops, InsPtr, BatchState, ExecCtx); }
+VM_FORCEINLINE const uint8 *    VVM_Exec1null_acquireindex                 (VVM_NULL_FN_ARGS)                                                                                                                   { return VVM_acquireindex(CT_MultipleLoops, InsPtr, BatchState, ExecCtx); }
+VM_FORCEINLINE const uint8 *    VVM_Exec0null_external_func_call           (VVM_NULL_FN_ARGS)                                                                                                                   { return VVM_external_func_call(CT_MultipleLoops, InsPtr, BatchState, ExecCtx, SerializeState, CmpSerializeState, NumLoops); }
+VM_FORCEINLINE const uint8 *    VVM_Exec1null_exec_index                   (VVM_NULL_FN_ARGS)                                                                                                                   { return VVM_exec_index(CT_MultipleLoops, InsPtr, BatchState, NumLoops); }
+VM_FORCEINLINE const uint8 *    VVM_Exec0null_noise2D                      (VVM_NULL_FN_ARGS)                                                                                                                   { return InsPtr; }
+VM_FORCEINLINE const uint8 *    VVM_Exec0null_noise3D                      (VVM_NULL_FN_ARGS)                                                                                                                   { return InsPtr; }
+VM_FORCEINLINE const uint8 *    VVM_Exec1null_enter_stat_scope             (VVM_NULL_FN_ARGS)                                                                                                                   { return InsPtr + 2; }
+VM_FORCEINLINE const uint8 *    VVM_Exec0null_exit_stat_scope              (VVM_NULL_FN_ARGS)                                                                                                                   { return InsPtr; }
+VM_FORCEINLINE const uint8 *    VVM_Exec0null_update_id                    (VVM_NULL_FN_ARGS)                                                                                                                   { return VVM_update_id(CT_MultipleLoops, InsPtr, BatchState, ExecCtx, NumLoops); }
+VM_FORCEINLINE const uint8 *    VVM_Exec1null_acquire_id                   (VVM_NULL_FN_ARGS)                                                                                                                   { return VVM_acquire_id(CT_MultipleLoops, InsPtr, BatchState, ExecCtx, NumLoops); }
+VM_FORCEINLINE const uint8 *    VVM_Exec1null_half_to_float                (VVM_NULL_FN_ARGS)                                                                                                                   { return VVM_half_to_float(CT_MultipleLoops, InsPtr, BatchState, NumLoops); }
+VM_FORCEINLINE const uint8 *    VVM_Exec1null_fasi                         (VVM_NULL_FN_ARGS)                                                                                                                   { return InsPtr; }
+VM_FORCEINLINE const uint8 *    VVM_Exec1null_iasf                         (VVM_NULL_FN_ARGS)                                                                                                                   { return InsPtr; }
+VM_FORCEINLINE const uint8 *    VVM_Exec1null_exec_indexf                  (VVM_NULL_FN_ARGS)                                                                                                                   { return VVM_exec_indexf(CT_MultipleLoops, InsPtr, BatchState, NumLoops); }
+VM_FORCEINLINE const uint8 *    VVM_Exec1null_exec_index_addi              (VVM_NULL_FN_ARGS)                                                                                                                   { return VVM_exec_index_addi(CT_MultipleLoops, InsPtr, BatchState, NumLoops); }
+VM_FORCEINLINE VectorRegister4f VVM_Exec4f_cmplt_select                    (FVectorVMBatchState *BatchState, VectorRegister4f a, VectorRegister4f b, VectorRegister4f c, VectorRegister4f d)                    { return VectorSelect(VectorCompareLT(a, b), c, d); }
+VM_FORCEINLINE VectorRegister4f VVM_Exec4f_cmple_select                    (FVectorVMBatchState *BatchState, VectorRegister4f a, VectorRegister4f b, VectorRegister4f c, VectorRegister4f d)                    { return VectorSelect(VectorCompareLE(a, b), c, d); }
+VM_FORCEINLINE VectorRegister4f VVM_Exec4f_cmpeq_select                    (FVectorVMBatchState *BatchState, VectorRegister4f a, VectorRegister4f b, VectorRegister4f c, VectorRegister4f d)                    { return VectorSelect(VectorCompareEQ(a, b), c, d); }
+VM_FORCEINLINE VectorRegister4f VVM_Exec4f_cmplti_select                   (FVectorVMBatchState *BatchState, VectorRegister4f a, VectorRegister4f b, VectorRegister4f c, VectorRegister4f d)                    { return VectorSelect(VectorCastIntToFloat(VectorIntCompareLT(*(VectorRegister4i *)&a, *(VectorRegister4i *)&b)), c, d); }
+VM_FORCEINLINE VectorRegister4f VVM_Exec4f_cmplei_select                   (FVectorVMBatchState *BatchState, VectorRegister4f a, VectorRegister4f b, VectorRegister4f c, VectorRegister4f d)                    { return VectorSelect(VectorCastIntToFloat(VectorIntCompareLE(*(VectorRegister4i *)&a, *(VectorRegister4i *)&b)), c, d); }
+VM_FORCEINLINE VectorRegister4f VVM_Exec4f_cmpeqi_select                   (FVectorVMBatchState *BatchState, VectorRegister4f a, VectorRegister4f b, VectorRegister4f c, VectorRegister4f d)                    { return VectorSelect(VectorCastIntToFloat(VectorIntCompareEQ(*(VectorRegister4i *)&a, *(VectorRegister4i *)&b)), c, d); }
+VM_FORCEINLINE VectorRegister4f VVM_Exec3f_cmplt_logic_and                 (FVectorVMBatchState *BatchState, VectorRegister4f a, VectorRegister4f b, VectorRegister4f c)                                        { return VectorCastIntToFloat(VectorIntAnd(VectorCastFloatToInt(VectorCompareLT(a, b)), *(VectorRegister4i *)&c)); }
+VM_FORCEINLINE VectorRegister4f VVM_Exec3f_cmple_logic_and                 (FVectorVMBatchState *BatchState, VectorRegister4f a, VectorRegister4f b, VectorRegister4f c)                                        { return VectorCastIntToFloat(VectorIntAnd(VectorCastFloatToInt(VectorCompareLE(a, b)), *(VectorRegister4i *)&c)); }
+VM_FORCEINLINE VectorRegister4f VVM_Exec3f_cmpgt_logic_and                 (FVectorVMBatchState *BatchState, VectorRegister4f a, VectorRegister4f b, VectorRegister4f c)                                        { return VectorCastIntToFloat(VectorIntAnd(VectorCastFloatToInt(VectorCompareGT(a, b)), *(VectorRegister4i *)&c)); }
+VM_FORCEINLINE VectorRegister4f VVM_Exec3f_cmpge_logic_and                 (FVectorVMBatchState *BatchState, VectorRegister4f a, VectorRegister4f b, VectorRegister4f c)                                        { return VectorCastIntToFloat(VectorIntAnd(VectorCastFloatToInt(VectorCompareGE(a, b)), *(VectorRegister4i *)&c)); }
+VM_FORCEINLINE VectorRegister4f VVM_Exec3f_cmpeq_logic_and                 (FVectorVMBatchState *BatchState, VectorRegister4f a, VectorRegister4f b, VectorRegister4f c)                                        { return VectorCastIntToFloat(VectorIntAnd(VectorCastFloatToInt(VectorCompareEQ(a, b)), *(VectorRegister4i *)&c)); }
+VM_FORCEINLINE VectorRegister4f VVM_Exec3f_cmpne_logic_and                 (FVectorVMBatchState *BatchState, VectorRegister4f a, VectorRegister4f b, VectorRegister4f c)                                        { return VectorCastIntToFloat(VectorIntAnd(VectorCastFloatToInt(VectorCompareNE(a, b)), *(VectorRegister4i *)&c)); }
+VM_FORCEINLINE VectorRegister4i VVM_Exec3i_cmplti_logic_and                (FVectorVMBatchState *BatchState, VectorRegister4i a, VectorRegister4i b, VectorRegister4i c)                                        { return VectorIntAnd(VectorIntCompareLT(a, b), c); }
+VM_FORCEINLINE VectorRegister4i VVM_Exec3i_cmplei_logic_and                (FVectorVMBatchState *BatchState, VectorRegister4i a, VectorRegister4i b, VectorRegister4i c)                                        { return VectorIntAnd(VectorIntCompareLE(a, b), c); }
+VM_FORCEINLINE VectorRegister4i VVM_Exec3i_cmpgti_logic_and                (FVectorVMBatchState *BatchState, VectorRegister4i a, VectorRegister4i b, VectorRegister4i c)                                        { return VectorIntAnd(VectorIntCompareGT(a, b), c); }
+VM_FORCEINLINE VectorRegister4i VVM_Exec3i_cmpgei_logic_and                (FVectorVMBatchState *BatchState, VectorRegister4i a, VectorRegister4i b, VectorRegister4i c)                                        { return VectorIntAnd(VectorIntCompareGE(a, b), c); }
+VM_FORCEINLINE VectorRegister4i VVM_Exec3i_cmpeqi_logic_and                (FVectorVMBatchState *BatchState, VectorRegister4i a, VectorRegister4i b, VectorRegister4i c)                                        { return VectorIntAnd(VectorIntCompareEQ(a, b), c); }
+VM_FORCEINLINE VectorRegister4i VVM_Exec3i_cmpnei_logic_and                (FVectorVMBatchState *BatchState, VectorRegister4i a, VectorRegister4i b, VectorRegister4i c)                                        { return VectorIntAnd(VectorIntCompareNEQ(a, b), c); }
+VM_FORCEINLINE VectorRegister4f VVM_Exec3f_cmplt_logic_or                  (FVectorVMBatchState *BatchState, VectorRegister4f a, VectorRegister4f b, VectorRegister4f c)                                        { return VectorCastIntToFloat(VectorIntOr(VectorCastFloatToInt(VectorCompareLT(a, b)), *(VectorRegister4i *)&c)); }
+VM_FORCEINLINE VectorRegister4f VVM_Exec3f_cmple_logic_or                  (FVectorVMBatchState *BatchState, VectorRegister4f a, VectorRegister4f b, VectorRegister4f c)                                        { return VectorCastIntToFloat(VectorIntOr(VectorCastFloatToInt(VectorCompareLE(a, b)), *(VectorRegister4i *)&c)); }
+VM_FORCEINLINE VectorRegister4f VVM_Exec3f_cmpgt_logic_or                  (FVectorVMBatchState *BatchState, VectorRegister4f a, VectorRegister4f b, VectorRegister4f c)                                        { return VectorCastIntToFloat(VectorIntOr(VectorCastFloatToInt(VectorCompareGT(a, b)), *(VectorRegister4i *)&c)); }
+VM_FORCEINLINE VectorRegister4f VVM_Exec3f_cmpge_logic_or                  (FVectorVMBatchState *BatchState, VectorRegister4f a, VectorRegister4f b, VectorRegister4f c)                                        { return VectorCastIntToFloat(VectorIntOr(VectorCastFloatToInt(VectorCompareGE(a, b)), *(VectorRegister4i *)&c)); }
+VM_FORCEINLINE VectorRegister4f VVM_Exec3f_cmpeq_logic_or                  (FVectorVMBatchState *BatchState, VectorRegister4f a, VectorRegister4f b, VectorRegister4f c)                                        { return VectorCastIntToFloat(VectorIntOr(VectorCastFloatToInt(VectorCompareEQ(a, b)), *(VectorRegister4i *)&c)); }
+VM_FORCEINLINE VectorRegister4f VVM_Exec3f_cmpne_logic_or                  (FVectorVMBatchState *BatchState, VectorRegister4f a, VectorRegister4f b, VectorRegister4f c)                                        { return VectorCastIntToFloat(VectorIntOr(VectorCastFloatToInt(VectorCompareNE(a, b)), *(VectorRegister4i *)&c)); }
+VM_FORCEINLINE VectorRegister4i VVM_Exec3i_cmplti_logic_or                 (FVectorVMBatchState *BatchState, VectorRegister4i a, VectorRegister4i b, VectorRegister4i c)                                        { return VectorIntOr(VectorIntCompareLT(a, b), c); }
+VM_FORCEINLINE VectorRegister4i VVM_Exec3i_cmplei_logic_or                 (FVectorVMBatchState *BatchState, VectorRegister4i a, VectorRegister4i b, VectorRegister4i c)                                        { return VectorIntOr(VectorIntCompareLE(a, b), c); }
+VM_FORCEINLINE VectorRegister4i VVM_Exec3i_cmpgti_logic_or                 (FVectorVMBatchState *BatchState, VectorRegister4i a, VectorRegister4i b, VectorRegister4i c)                                        { return VectorIntOr(VectorIntCompareGT(a, b), c); }
+VM_FORCEINLINE VectorRegister4i VVM_Exec3i_cmpgei_logic_or                 (FVectorVMBatchState *BatchState, VectorRegister4i a, VectorRegister4i b, VectorRegister4i c)                                        { return VectorIntOr(VectorIntCompareGE(a, b), c); }
+VM_FORCEINLINE VectorRegister4i VVM_Exec3i_cmpeqi_logic_or                 (FVectorVMBatchState *BatchState, VectorRegister4i a, VectorRegister4i b, VectorRegister4i c)                                        { return VectorIntOr(VectorIntCompareEQ(a, b), c); }
+VM_FORCEINLINE VectorRegister4i VVM_Exec3i_cmpnei_logic_or                 (FVectorVMBatchState *BatchState, VectorRegister4i a, VectorRegister4i b, VectorRegister4i c)                                        { return VectorIntOr(VectorIntCompareNEQ(a, b), c); }
+VM_FORCEINLINE VectorRegister4f VVM_Exec4f_mad_add                         (FVectorVMBatchState *BatchState, VectorRegister4f a, VectorRegister4f b, VectorRegister4f c, VectorRegister4f d)                    { return VectorAdd(VectorMultiplyAdd(a, b, c), d); }
+VM_FORCEINLINE VectorRegister4f VVM_Exec4f_mad_sub0                        (FVectorVMBatchState *BatchState, VectorRegister4f a, VectorRegister4f b, VectorRegister4f c, VectorRegister4f d)                    { return VectorSubtract(VectorMultiplyAdd(a, b, c), d); }
+VM_FORCEINLINE VectorRegister4f VVM_Exec4f_mad_sub1                        (FVectorVMBatchState *BatchState, VectorRegister4f a, VectorRegister4f b, VectorRegister4f c, VectorRegister4f d)                    { return VectorSubtract(d, VectorMultiplyAdd(a, b, c)); }
+VM_FORCEINLINE VectorRegister4f VVM_Exec4f_mad_mul                         (FVectorVMBatchState *BatchState, VectorRegister4f a, VectorRegister4f b, VectorRegister4f c, VectorRegister4f d)                    { return VectorMultiply(VectorMultiplyAdd(a, b, c), d); }
+VM_FORCEINLINE VectorRegister4f VVM_Exec3f_mad_sqrt                        (FVectorVMBatchState *BatchState, VectorRegister4f a, VectorRegister4f b, VectorRegister4f c)                                        { return VectorSqrt(VectorMultiplyAdd(a, b, c)); }
+VM_FORCEINLINE VectorRegister4f VVM_Exec5f_mad_mad0                        (FVectorVMBatchState *BatchState, VectorRegister4f a, VectorRegister4f b, VectorRegister4f c, VectorRegister4f d, VectorRegister4f e){ return VectorMultiplyAdd(d, e, VectorMultiplyAdd(a, b, c)); }
+VM_FORCEINLINE VectorRegister4f VVM_Exec5f_mad_mad1                        (FVectorVMBatchState *BatchState, VectorRegister4f a, VectorRegister4f b, VectorRegister4f c, VectorRegister4f d, VectorRegister4f e){ return VectorMultiplyAdd(VectorMultiplyAdd(a, b, c), d, e); }
+VM_FORCEINLINE VectorRegister4f VVM_Exec4f_mul_mad0                        (FVectorVMBatchState *BatchState, VectorRegister4f a, VectorRegister4f b, VectorRegister4f c, VectorRegister4f d)                    { return VectorMultiplyAdd(VectorMultiply(a, b), c, d); }
+VM_FORCEINLINE VectorRegister4f VVM_Exec4f_mul_mad1                        (FVectorVMBatchState *BatchState, VectorRegister4f a, VectorRegister4f b, VectorRegister4f c, VectorRegister4f d)                    { return VectorMultiplyAdd(c, d, VectorMultiply(a, b)); }
+VM_FORCEINLINE VectorRegister4f VVM_Exec3f_mul_add                         (FVectorVMBatchState *BatchState, VectorRegister4f a, VectorRegister4f b, VectorRegister4f c)                                        { return VectorAdd(VectorMultiply(a, b), c); }
+VM_FORCEINLINE VectorRegister4f VVM_Exec3f_mul_sub0                        (FVectorVMBatchState *BatchState, VectorRegister4f a, VectorRegister4f b, VectorRegister4f c)                                        { return VectorSubtract(VectorMultiply(a, b), c); }
+VM_FORCEINLINE VectorRegister4f VVM_Exec3f_mul_sub1                        (FVectorVMBatchState *BatchState, VectorRegister4f a, VectorRegister4f b, VectorRegister4f c)                                        { return VectorSubtract(c, VectorMultiply(a, b)); }
+VM_FORCEINLINE VectorRegister4f VVM_Exec3f_mul_mul                         (FVectorVMBatchState *BatchState, VectorRegister4f a, VectorRegister4f b, VectorRegister4f c)                                        { return VectorMultiply(VectorMultiply(a, b), c); }
+VM_FORCEINLINE VectorRegister4f VVM_Exec3f_mul_max                         (FVectorVMBatchState *BatchState, VectorRegister4f a, VectorRegister4f b, VectorRegister4f c)                                        { return VectorMax(VectorMultiply(a, b), c); }
+VM_FORCEINLINE VectorRegister4f VVM_Exec2f_mul_2x                          (FVectorVMBatchState *BatchState, VectorRegister4f a, VectorRegister4f b)                                                            { return VectorMultiply(a, b); }
+VM_FORCEINLINE VectorRegister4f VVM_Exec4f_add_mad1                        (FVectorVMBatchState *BatchState, VectorRegister4f a, VectorRegister4f b, VectorRegister4f c, VectorRegister4f d)                    { return VectorMultiplyAdd(c, d, VectorAdd(a, b)); }
+VM_FORCEINLINE VectorRegister4f VVM_Exec3f_add_add                         (FVectorVMBatchState *BatchState, VectorRegister4f a, VectorRegister4f b, VectorRegister4f c)                                        { return VectorAdd(VectorAdd(a, b), c); }
+VM_FORCEINLINE VectorRegister4f VVM_Exec3f_sub_cmplt1                      (FVectorVMBatchState *BatchState, VectorRegister4f a, VectorRegister4f b, VectorRegister4f c)                                        { return VectorCompareLT(c, VectorSubtract(a, b)); }
+VM_FORCEINLINE VectorRegister4f VVM_Exec2f_sub_neg                         (FVectorVMBatchState *BatchState, VectorRegister4f a, VectorRegister4f b)                                                            { return VectorNegate(VectorSubtract(a, b)); }
+VM_FORCEINLINE VectorRegister4f VVM_Exec3f_sub_mul                         (FVectorVMBatchState *BatchState, VectorRegister4f a, VectorRegister4f b, VectorRegister4f c)                                        { return VectorMultiply(VectorSubtract(a, b), c); }
+VM_FORCEINLINE VectorRegister4f VVM_Exec4f_div_mad0                        (FVectorVMBatchState *BatchState, VectorRegister4f a, VectorRegister4f b, VectorRegister4f c, VectorRegister4f d)                    { return VectorMultiplyAdd(VVM_Exec2f_div(BatchState, a, b), c, d); }
+VM_FORCEINLINE VectorRegister4i VVM_Exec2i_div_f2i                         (FVectorVMBatchState *BatchState, VectorRegister4i a, VectorRegister4i b)                                                            { return VectorFloatToInt(VVM_Exec2f_div(BatchState, VectorCastIntToFloat(a), VectorCastIntToFloat(b))); }
+VM_FORCEINLINE VectorRegister4f VVM_Exec3f_div_mul                         (FVectorVMBatchState *BatchState, VectorRegister4f a, VectorRegister4f b, VectorRegister4f c)                                        { return VectorMultiply(VVM_Exec2f_div(BatchState, a, b), c); }
+VM_FORCEINLINE VectorRegister4i VVM_Exec3i_muli_addi                       (FVectorVMBatchState *BatchState, VectorRegister4i a, VectorRegister4i b, VectorRegister4i c)                                        { return VectorIntAdd(VectorIntMultiply(a, b), c); }
+VM_FORCEINLINE VectorRegister4i VVM_Exec3i_addi_bit_rshift                 (FVectorVMBatchState *BatchState, VectorRegister4i a, VectorRegister4i b, VectorRegister4i c)                                        { return VVMIntRShift(VectorIntAdd(a, b), c); }
+VM_FORCEINLINE VectorRegister4i VVM_Exec3i_addi_muli                       (FVectorVMBatchState *BatchState, VectorRegister4i a, VectorRegister4i b, VectorRegister4i c)                                        { return VectorIntMultiply(VectorIntAdd(a, b), c); }
+VM_FORCEINLINE VectorRegister4i VVM_Exec1i_b2i_2x                          (FVectorVMBatchState *BatchState, VectorRegister4i a)                                                                                { return VectorIntSelect(a, VectorIntSet1(1), VectorSetZero()); }
+VM_FORCEINLINE VectorRegister4f VVM_Exec2f_i2f_div0                        (FVectorVMBatchState *BatchState, VectorRegister4f a, VectorRegister4f b)                                                            { return VVM_Exec2f_div(BatchState, VectorIntToFloat(VectorCastFloatToInt(a)), b); }
+VM_FORCEINLINE VectorRegister4f VVM_Exec2f_i2f_div1                        (FVectorVMBatchState *BatchState, VectorRegister4f a, VectorRegister4f b)                                                            { return VVM_Exec2f_div(BatchState, b, VectorIntToFloat(VectorCastFloatToInt(a))); }
+VM_FORCEINLINE VectorRegister4f VVM_Exec2f_i2f_mul                         (FVectorVMBatchState *BatchState, VectorRegister4f a, VectorRegister4f b)                                                            { return VectorMultiply(VectorIntToFloat(VectorCastFloatToInt(a)), b); }
+VM_FORCEINLINE VectorRegister4f VVM_Exec3f_i2f_mad0                        (FVectorVMBatchState *BatchState, VectorRegister4f a, VectorRegister4f b, VectorRegister4f c)                                        { return VectorMultiplyAdd(VectorIntToFloat(VectorCastFloatToInt(a)), b, c); }
+VM_FORCEINLINE VectorRegister4f VVM_Exec3f_i2f_mad1                        (FVectorVMBatchState *BatchState, VectorRegister4f a, VectorRegister4f b, VectorRegister4f c)                                        { return VectorMultiplyAdd(a, b, VectorIntToFloat(VectorCastFloatToInt(c))); }
+VM_FORCEINLINE VectorRegister4i VVM_Exec3i_f2i_select1                     (FVectorVMBatchState *BatchState, VectorRegister4i a, VectorRegister4i b, VectorRegister4i c)                                        { return VectorIntSelect(a, VectorFloatToInt(VectorCastIntToFloat(b)), c); }
+VM_FORCEINLINE VectorRegister4i VVM_Exec2i_f2i_maxi                        (FVectorVMBatchState *BatchState, VectorRegister4i a, VectorRegister4i b)                                                            { return VectorIntMax(VectorFloatToInt(VectorCastIntToFloat(a)), b); }
+VM_FORCEINLINE VectorRegister4i VVM_Exec2i_f2i_addi                        (FVectorVMBatchState *BatchState, VectorRegister4i a, VectorRegister4i b)                                                            { return VectorIntAdd(VectorFloatToInt(VectorCastIntToFloat(a)), b); }
+VM_FORCEINLINE VectorRegister4f VVM_Exec3f_fmod_add                        (FVectorVMBatchState *BatchState, VectorRegister4f a, VectorRegister4f b, VectorRegister4f c)                                        { return VectorAdd(VectorMod(a, b), c); }
+VM_FORCEINLINE VectorRegister4f VVM_Exec2f_bit_and_i2f                     (FVectorVMBatchState *BatchState, VectorRegister4f a, VectorRegister4f b)                                                            { return VectorIntToFloat(VectorIntAnd(VectorCastFloatToInt(a), VectorCastFloatToInt(b))); }
+VM_FORCEINLINE VectorRegister4i VVM_Exec3i_bit_rshift_bit_and              (FVectorVMBatchState *BatchState, VectorRegister4i a, VectorRegister4i b, VectorRegister4i c)                                        { return VectorIntAnd(VVMIntRShift(a, b), c); }
+VM_FORCEINLINE VectorRegister4f VVM_Exec2f_neg_cmplt                       (FVectorVMBatchState *BatchState, VectorRegister4f a, VectorRegister4f b)                                                            { return VectorCompareLT(VectorNegate(a), b); }
+VM_FORCEINLINE VectorRegister4i VVM_Exec3i_bit_or_muli                     (FVectorVMBatchState *BatchState, VectorRegister4i a, VectorRegister4i b, VectorRegister4i c)                                        { return VectorIntMultiply(VectorIntOr(a, b), c); }
+VM_FORCEINLINE VectorRegister4i VVM_Exec3i_bit_lshift_bit_or               (FVectorVMBatchState *BatchState, VectorRegister4i a, VectorRegister4i b, VectorRegister4i c)                                        { return VectorIntOr(VVMIntLShift(a, b), c); }
+VM_FORCEINLINE const uint8 *    VVM_Exec2null_random_add                   (VVM_NULL_FN_ARGS)                                                                                                                   { return VVM_random_add(CT_MultipleLoops, InsPtr, BatchState, ExecCtx, SerializeState, CmpSerializeState, NumLoops); }
+VM_FORCEINLINE const uint8 *    VVM_Exec1null_random_2x                    (VVM_NULL_FN_ARGS)                                                                                                                   { return VVM_random_2x(CT_MultipleLoops, InsPtr, BatchState, ExecCtx, SerializeState, CmpSerializeState, NumLoops); }
+VM_FORCEINLINE VectorRegister4i VVM_Exec2i_max_f2i                         (FVectorVMBatchState *BatchState, VectorRegister4i a, VectorRegister4i b)                                                            { return VectorFloatToInt(VectorMax(VectorCastIntToFloat(a), VectorCastIntToFloat(b))); }
+VM_FORCEINLINE VectorRegister4f VVM_Exec4f_select_mul                      (FVectorVMBatchState *BatchState, VectorRegister4f a, VectorRegister4f b, VectorRegister4f c, VectorRegister4f d)                    { return VectorMultiply(VectorSelect(a, b, c), d); }
+VM_FORCEINLINE VectorRegister4f VVM_Exec4f_select_add                      (FVectorVMBatchState *BatchState, VectorRegister4f a, VectorRegister4f b, VectorRegister4f c, VectorRegister4f d)                    { return VectorAdd(VectorSelect(a, b, c), d); }
+VM_FORCEINLINE const uint8 *    VVM_Exec1null_sin_cos                      (VVM_NULL_FN_ARGS)                                                                                                                   { return VVM_sincos(CT_MultipleLoops, InsPtr, BatchState, NumLoops); }
+VM_FORCEINLINE const uint8 *    VVM_Exec0null_inputdata_float              (VVM_NULL_FN_ARGS)                                                                                                                   { return InsPtr; }
+VM_FORCEINLINE const uint8 *    VVM_Exec0null_inputdata_int32              (VVM_NULL_FN_ARGS)                                                                                                                   { return InsPtr; }
+VM_FORCEINLINE const uint8 *    VVM_Exec0null_inputdata_half               (VVM_NULL_FN_ARGS)                                                                                                                   { return InsPtr; }
+VM_FORCEINLINE const uint8 *    VVM_Exec0null_inputdata_noadvance_float    (VVM_NULL_FN_ARGS)                                                                                                                   { return InsPtr; }
+VM_FORCEINLINE const uint8 *    VVM_Exec0null_inputdata_noadvance_int32    (VVM_NULL_FN_ARGS)                                                                                                                   { return InsPtr; }
+VM_FORCEINLINE const uint8 *    VVM_Exec0null_inputdata_noadvance_half     (VVM_NULL_FN_ARGS)                                                                                                                   { return InsPtr; }
+VM_FORCEINLINE const uint8 *    VVM_Exec0null_outputdata_float_from_half   (VVM_NULL_FN_ARGS)                                                                                                                   { return VVM_Output32_from_16(CT_MultipleLoops, InsPtr, BatchState, ExecCtx); }
+VM_FORCEINLINE const uint8 *    VVM_Exec0null_outputdata_half_from_half    (VVM_NULL_FN_ARGS)                                                                                                                   { return VVM_Output16_from_16(CT_MultipleLoops, InsPtr, BatchState, ExecCtx); }
 //this is the macro required to exit the infinite loop running over the bytecode most efficiently
 #define VVM_Dispatch_execFn0done_0done(...) NULL; goto done_loop;
 
@@ -2457,7 +2543,7 @@ static const void* jmp_tbl[] = {
 #else
 #define VVM_OP_START    switch ((EVectorVMOp)InsPtr[-1])
 #if VECTORVM_DEBUG_PRINTF
-#define VVM_OP_CASE(op)	case EVectorVMOp::op: printf("%d, %s", DbgPfInsCount++, VVM_OP_NAMES[InsPtr[-1]]);
+#define VVM_OP_CASE(op)	case EVectorVMOp::op: VECTORVM_PRINTF("%d, %s", DbgPfInsCount++, VVM_OP_NAMES_EXP[InsPtr[-1]]);
 #else
 #define VVM_OP_CASE(op)	case EVectorVMOp::op:
 #endif
@@ -2465,7 +2551,7 @@ static const void* jmp_tbl[] = {
 #endif
 
 #if VECTORVM_DEBUG_PRINTF
-static const char *VVM_OP_NAMES[] = {
+static const char *VVM_OP_NAMES_EXP[] = {
 #	define VVM_OP_XM(OpCode, ...) #OpCode,
 	VVM_OP_XM_LIST
 #	undef VVM_OP_XM
@@ -2775,4 +2861,3 @@ void ExecVectorVMState(struct FVectorVMState *VVMState, struct FVectorVMSerializ
 }
 
 #endif //NIAGARA_EXP_VM
-
diff --git a/Engine/Source/Runtime/VectorVM/Private/VectorVMExperimentalOptimizer.cpp b/Engine/Source/Runtime/VectorVM/Private/VectorVMExperimentalOptimizer.cpp
--- a/Engine/Source/Runtime/VectorVM/Private/VectorVMExperimentalOptimizer.cpp
+++ b/Engine/Source/Runtime/VectorVM/Private/VectorVMExperimentalOptimizer.cpp
@@ -763,6 +763,7 @@ static bool VVMDoesInstructionHaveDependencies(FVectorVMOptimizeContext *OptCont
 	}
 	return false;
 }
+
 #	define VVMPushRegUsage(RegIdx, Type, IOFlag)	          do { uint32 Res = VVMPushRegUsage_(OptContext, Instruction, RegIdx, Type, IOFlag, &NumRegisterUsageAlloced);                              if (Res != 0) { return Res; } } while(0);
 
 
@@ -1309,8 +1310,40 @@ uint32 OptimizeVectorVMScript(const uint8 *InBytecode, int InBytecodeLen, FVecto
 		}
 	}
 
+	//Step 7: remove cast instructions and re-map the inputs of instructions that the casts feed into
+	//to the inputs of the cast instructions.
+	for (uint32 CastInsIdx = 0; CastInsIdx < OptContext->Intermediate.NumInstructions; ++CastInsIdx) {
+		FVectorVMOptimizeInstruction *CastIns = OptContext->Intermediate.Instructions + CastInsIdx;
+		if (CastIns->OpCode == EVectorVMOp::iasf || CastIns->OpCode == EVectorVMOp::fasi) {
+			uint16 CastInputSSAReg   = OptContext->Intermediate.SSARegisterUsageBuffer[CastIns->RegPtrOffset + 0];
+			uint16 CastOutputSSAReg  = OptContext->Intermediate.SSARegisterUsageBuffer[CastIns->RegPtrOffset + 1];
+			uint16 CastInputReg      = OptContext->Intermediate.RegisterUsageBuffer[CastIns->RegPtrOffset + 0];
+			uint16 CastOutputReg     = OptContext->Intermediate.RegisterUsageBuffer[CastIns->RegPtrOffset + 1];
+			uint8  CastInputRegType  = OptContext->Intermediate.RegisterUsageType[CastIns->RegPtrOffset + 0];
+			uint8  CastOutputRegType = OptContext->Intermediate.RegisterUsageType[CastIns->RegPtrOffset + 1];
+			for (uint32 InsIdx = CastInsIdx + 1; InsIdx < OptContext->Intermediate.NumInstructions; ++InsIdx) {
+				FVectorVMOptimizeInstruction *Ins = OptContext->Intermediate.Instructions + InsIdx;
+				for (int i = 0; i < Ins->NumInputRegisters; ++i) {
+					uint16 SSAReg  = OptContext->Intermediate.SSARegisterUsageBuffer[Ins->RegPtrOffset + i];
+					uint16 Reg     = OptContext->Intermediate.RegisterUsageBuffer   [Ins->RegPtrOffset + i];
+					uint8  RegType = OptContext->Intermediate.RegisterUsageType     [Ins->RegPtrOffset + i];
+					if (SSAReg == CastOutputSSAReg && RegType == CastOutputRegType) {
+						check(RegType == VVM_RT_TEMPREG);
+						OptContext->Intermediate.SSARegisterUsageBuffer[Ins->RegPtrOffset + i] = CastInputSSAReg;
+						OptContext->Intermediate.RegisterUsageBuffer   [Ins->RegPtrOffset + i] = CastInputReg;
+						OptContext->Intermediate.RegisterUsageType     [Ins->RegPtrOffset + i] = CastInputRegType;
+					}
+				}
+			}
+			//remove this cast instruction
+			FMemory::Memmove(OptContext->Intermediate.Instructions + CastInsIdx, OptContext->Intermediate.Instructions + CastInsIdx + 1, sizeof(FVectorVMOptimizeInstruction) * (OptContext->Intermediate.NumInstructions - CastInsIdx - 1));
+			--OptContext->Intermediate.NumInstructions;
+			--CastInsIdx;
+		}
+	}
+
 	int OnePastLastInputIdx = -1;
-	{ //Step 7: change temp registers that come directly from inputs to the input index
+	{ //Step 8: change temp registers that come directly from inputs to the input index
 		for (uint32 i = 0; i < OptContext->Intermediate.NumInstructions; ++i)
 		{
 			FVectorVMOptimizeInstruction *InputIns = OptContext->Intermediate.Instructions + i;
@@ -1336,6 +1369,16 @@ uint32 OptimizeVectorVMScript(const uint8 *InBytecode, int InBytecodeLen, FVecto
 							check(InputIns->Index == i); //make sure the instruction is in its correct place.  Instructions could have moved, but this should have been corrected above.
 							OptContext->Intermediate.SSARegisterUsageBuffer[Ins->RegPtrOffset + 1] = InputIns->Index;
 							OptContext->Intermediate.RegisterUsageType[Ins->RegPtrOffset + 1]      = VVM_RT_INPUT;
+
+							//deal with output instructions that read directly from inputs where there's a required half/float conversion
+							if (InputIns->OpCode == EVectorVMOp::inputdata_half && Ins->OpCode == EVectorVMOp::outputdata_float)
+							{ //input half, output 32 bit
+								Ins->OpCode = EVectorVMOp::outputdata_float_from_half;
+							}
+							else if (InputIns->OpCode == EVectorVMOp::inputdata_half && Ins->OpCode == EVectorVMOp::outputdata_half)
+							{
+								Ins->OpCode = EVectorVMOp::outputdata_half_from_half;
+							}
 						}
 					break;
 					default:
@@ -1419,7 +1462,7 @@ uint32 OptimizeVectorVMScript(const uint8 *InBytecode, int InBytecodeLen, FVecto
 
 		int LowestInstructionIdxForAcquireIdx = OnePastLastInputIdx; //acquire index instructions will be sorted by whichever comes first in the IR... possibly worth checking if re-ordering is more efficient
 
-		{ //Step 8: Find all the acquireindex instructions and re-order them to be executed ASAP
+		{ //Step 9: Find all the acquireindex instructions and re-order them to be executed ASAP
 			int NumAcquireIndexInstructions = 0;
 			for (uint32 i = 0; i < OptContext->Intermediate.NumInstructions; ++i)
 			{
@@ -1463,7 +1506,7 @@ uint32 OptimizeVectorVMScript(const uint8 *InBytecode, int InBytecodeLen, FVecto
 			}
 		}
 
-		{ //Step 9: re-order the outputs to be done as early as possible: after the SSA's register's last usage
+		{ //Step 10: re-order the outputs to be done as early as possible: after the SSA's register's last usage
 			for (uint32 OutputInsIdx = 0; OutputInsIdx < OptContext->Intermediate.NumInstructions; ++OutputInsIdx)
 			{
 				FVectorVMOptimizeInstruction *OutputIns = OptContext->Intermediate.Instructions + OutputInsIdx;
@@ -1502,7 +1545,7 @@ uint32 OptimizeVectorVMScript(const uint8 *InBytecode, int InBytecodeLen, FVecto
 			}
 		}
 
-		{ //Step 10: re-order all dependent-less instructions to right before their output is used
+		{ //Step 11: re-order all dependent-less instructions to right before their output is used
 			uint16 LastSwapInstructionIdx = 0xFFFF; //to prevent an infinite loop when one instruction has two or more dependencies and they keep swapping back and forth
 			for (uint32 i = 0; i < OptContext->Intermediate.NumInstructions; ++i)
 			{
@@ -1535,8 +1578,8 @@ uint32 OptimizeVectorVMScript(const uint8 *InBytecode, int InBytecodeLen, FVecto
 			}
 		}
 	}
-	
-	{ //Step 11: group and sort all output instructions
+
+	{ //Step 12: group and sort all output instructions
 		for (uint32 i = 0; i < OptContext->Intermediate.NumInstructions; ++i)
 		{
 			FVectorVMOptimizeInstruction *InsStart = OptContext->Intermediate.Instructions + i;
@@ -1584,7 +1627,7 @@ uint32 OptimizeVectorVMScript(const uint8 *InBytecode, int InBytecodeLen, FVecto
 	
 	bool ChangedMergedIns = false;
 
-	{ //Step 12: figure out which instructions can be merged
+	{ //Step 13: figure out which instructions can be merged
 		struct FMergableOp
 		{
 			uint32		InsIdx0;
@@ -1787,7 +1830,7 @@ uint32 OptimizeVectorVMScript(const uint8 *InBytecode, int InBytecodeLen, FVecto
 		}
 		
 		
-		//Step 13: Merge all the statistically common pairs of instructions
+		//Step 14: Merge all the statistically common pairs of instructions
 #		define VVMCreateNewRegVars(NumNewRegs)	VVMEnsureRegAlloced(OptContext, true, &NumRegisterUsageAlloced, (NumNewRegs));                                        \
 												uint16 *Ins0Regs       = OptContext->Intermediate.RegisterUsageBuffer + Ins0->RegPtrOffset;                           \
 												uint16 *Ins1Regs       = OptContext->Intermediate.RegisterUsageBuffer + Ins1->RegPtrOffset;                           \
@@ -2509,7 +2552,7 @@ uint32 OptimizeVectorVMScript(const uint8 *InBytecode, int InBytecodeLen, FVecto
 	int NumInputRegisters  = 0;
 	int NumOutputRegisters = 0;
 
-	{ //Step 14: generate input remap table and output remaps
+	{ //Step 15: generate input remap table and output remaps
 		int NumInputRegistersAlloced  = 256;
 		int NumOutputRegistersAlloced = 256;
 
@@ -2727,7 +2770,7 @@ uint32 OptimizeVectorVMScript(const uint8 *InBytecode, int InBytecodeLen, FVecto
 		}
 	}
 	
-	{ //Step 15: use the SSA registers to compute the minimized registers required and write them back into the register usage buffer		
+	{ //Step 16: use the SSA registers to compute the minimized registers required and write them back into the register usage buffer		
 		uint16 *SSAUseMap2    = (uint16 *)OptContext->Init.ReallocFn(nullptr, sizeof(uint16) * NumSSARegistersUsed * 2, __FILE__, __LINE__);
 		if (SSAUseMap2 == nullptr)
 		{
@@ -2804,7 +2847,7 @@ uint32 OptimizeVectorVMScript(const uint8 *InBytecode, int InBytecodeLen, FVecto
 		OptContext->Init.FreeFn(SSAUseMap2, __FILE__, __LINE__);
 	}
 	
-	//Step 16 alias temp registers to output registers 
+	//Step 17 alias temp registers to output registers 
 	int NumRemappedOutputRegs   = FMath::Max(OptContext->NumOutputsRemapped, (uint16)OptContext->NumTempRegisters);
 	uint16 *FinalTempRegRemap = (uint16 *)OptContext->Init.ReallocFn(nullptr, sizeof(uint16) * NumRemappedOutputRegs, __FILE__, __LINE__);
 	for (int i = 0; i < NumRemappedOutputRegs; ++i)
@@ -2927,7 +2970,7 @@ uint32 OptimizeVectorVMScript(const uint8 *InBytecode, int InBytecodeLen, FVecto
 		OptContext->Init.FreeFn(AvailableOutputs, __FILE__, __LINE__);
 	}
 		
-	{ //Step 17: write the final optimized bytecode
+	{ //Step 18: write the final optimized bytecode
 		//this goes over the instruction list twice.  The first time to figure out how many bytes are required for the bytecode, the second to write the bytecode.
 		uint8 *OptimizedBytecode = nullptr;		
 		int NumOptimizedBytesRequired = 0;
diff --git a/Engine/Source/Runtime/VectorVM/Public/VectorVMCommon.h b/Engine/Source/Runtime/VectorVM/Public/VectorVMCommon.h
--- a/Engine/Source/Runtime/VectorVM/Public/VectorVMCommon.h
+++ b/Engine/Source/Runtime/VectorVM/Public/VectorVMCommon.h
@@ -130,89 +130,92 @@ enum class EVectorVMOperandLocation : uint8
 	VVM_OP_XM( half_to_float                , Op        , 1    , 1   , null     , 0             , 0               , VVM_INS_PARAM_FFFFFF) /* 95 */           \
 	VVM_OP_XM( fasi                         , Op        , 1    , 1   , null     , 0             , 0               , VVM_INS_PARAM_FFFFFF) /* 96 */           \
 	VVM_OP_XM( iasf                         , Op        , 1    , 1   , null     , 0             , 0               , VVM_INS_PARAM_FFFFFI) /* 97 */           \
-	/* Merged ops -- combined ops that show up frequently together in Fornite.  There are three types: */   \
-	/* 1. exec_index that get immediately fed into an add or i2f                                       */   \
-	/* 2. ops with identical inputs                                                                    */   \
-	/* 3. ops where the output of one chain to the next.. ie a mul that feeds directly into a sub      */   \
-	VVM_OP_XM( exec_indexf                  , Op        , 1    , 1   , null   , 0    , 0   , VVM_INS_PARAM_FFIIIF) /* 96 */            \
-	VVM_OP_XM( exec_index_addi              , Op        , 1    , 1   , null   , 0    , 0   , VVM_INS_PARAM_FFIIII) /* 97 */            \
-	VVM_OP_XM( cmplt_select                 , Op        , 4    , 1   , f      , 0    , 0   , VVM_INS_PARAM_FFFFFF) /* 98 */            \
-	VVM_OP_XM( cmple_select                 , Op        , 4    , 1   , f      , 0    , 0   , VVM_INS_PARAM_FFFFFF) /* 99 */            \
-	VVM_OP_XM( cmpeq_select                 , Op        , 4    , 1   , f      , 0    , 0   , VVM_INS_PARAM_FFFFFF) /* 100 */           \
-	VVM_OP_XM( cmplti_select                , Op        , 4    , 1   , f      , 0    , 0   , VVM_INS_PARAM_FIIFFF) /* 101 */           \
-	VVM_OP_XM( cmplei_select                , Op        , 4    , 1   , f      , 0    , 0   , VVM_INS_PARAM_FIIFFF) /* 102 */           \
-	VVM_OP_XM( cmpeqi_select                , Op        , 4    , 1   , f      , 0    , 0   , VVM_INS_PARAM_FIIFFF) /* 103 */           \
-	VVM_OP_XM( cmplt_logic_and              , Op        , 3    , 1   , f      , 0    , 0   , VVM_INS_PARAM_FFFFII) /* 104 */           \
-	VVM_OP_XM( cmple_logic_and              , Op        , 3    , 1   , f      , 0    , 0   , VVM_INS_PARAM_FFFFII) /* 105 */           \
-	VVM_OP_XM( cmpgt_logic_and              , Op        , 3    , 1   , f      , 0    , 0   , VVM_INS_PARAM_FFFFII) /* 106 */           \
-	VVM_OP_XM( cmpge_logic_and              , Op        , 3    , 1   , f      , 0    , 0   , VVM_INS_PARAM_FFFFII) /* 107 */           \
-	VVM_OP_XM( cmpeq_logic_and              , Op        , 3    , 1   , f      , 0    , 0   , VVM_INS_PARAM_FFFFII) /* 108 */           \
-	VVM_OP_XM( cmpne_logic_and              , Op        , 3    , 1   , f      , 0    , 0   , VVM_INS_PARAM_FFFFII) /* 109 */           \
-	VVM_OP_XM( cmplti_logic_and             , Op        , 3    , 1   , i      , 0    , 0   , VVM_INS_PARAM_FIIFII) /* 110 */           \
-	VVM_OP_XM( cmplei_logic_and             , Op        , 3    , 1   , i      , 0    , 0   , VVM_INS_PARAM_FIIFII) /* 111 */           \
-	VVM_OP_XM( cmpgti_logic_and             , Op        , 3    , 1   , i      , 0    , 0   , VVM_INS_PARAM_FIIFII) /* 112 */           \
-	VVM_OP_XM( cmpgei_logic_and             , Op        , 3    , 1   , i      , 0    , 0   , VVM_INS_PARAM_FIIFIF) /* 113 */           \
-	VVM_OP_XM( cmpeqi_logic_and             , Op        , 3    , 1   , i      , 0    , 0   , VVM_INS_PARAM_FIIFIF) /* 114 */           \
-	VVM_OP_XM( cmpnei_logic_and             , Op        , 3    , 1   , i      , 0    , 0   , VVM_INS_PARAM_FIIFIF) /* 115 */           \
-	VVM_OP_XM( cmplt_logic_or               , Op        , 3    , 1   , f      , 0    , 0   , VVM_INS_PARAM_FFFFFF) /* 116 */           \
-	VVM_OP_XM( cmple_logic_or               , Op        , 3    , 1   , f      , 0    , 0   , VVM_INS_PARAM_FFFFFF) /* 117 */           \
-	VVM_OP_XM( cmpgt_logic_or               , Op        , 3    , 1   , f      , 0    , 0   , VVM_INS_PARAM_FFFFFF) /* 118 */           \
-	VVM_OP_XM( cmpge_logic_or               , Op        , 3    , 1   , f      , 0    , 0   , VVM_INS_PARAM_FFFFFF) /* 119 */           \
-	VVM_OP_XM( cmpeq_logic_or               , Op        , 3    , 1   , f      , 0    , 0   , VVM_INS_PARAM_FFFFFF) /* 120 */           \
-	VVM_OP_XM( cmpne_logic_or               , Op        , 3    , 1   , f      , 0    , 0   , VVM_INS_PARAM_FFFFFF) /* 121 */           \
-	VVM_OP_XM( cmplti_logic_or              , Op        , 3    , 1   , i      , 0    , 0   , VVM_INS_PARAM_FIIFFF) /* 122 */           \
-	VVM_OP_XM( cmplei_logic_or              , Op        , 3    , 1   , i      , 0    , 0   , VVM_INS_PARAM_FIIFFF) /* 123 */           \
-	VVM_OP_XM( cmpgti_logic_or              , Op        , 3    , 1   , i      , 0    , 0   , VVM_INS_PARAM_FIIFFF) /* 124 */           \
-	VVM_OP_XM( cmpgei_logic_or              , Op        , 3    , 1   , i      , 0    , 0   , VVM_INS_PARAM_FIIFFF) /* 125 */           \
-	VVM_OP_XM( cmpeqi_logic_or              , Op        , 3    , 1   , i      , 0    , 0   , VVM_INS_PARAM_FIIFFF) /* 126 */           \
-	VVM_OP_XM( cmpnei_logic_or              , Op        , 3    , 1   , i      , 0    , 0   , VVM_INS_PARAM_FIIFFF) /* 127 */           \
-	VVM_OP_XM( mad_add                      , Op        , 4    , 1   , f      , 0    , 0   , VVM_INS_PARAM_FFFFFF) /* 128 */           \
-	VVM_OP_XM( mad_sub0                     , Op        , 4    , 1   , f      , 0    , 0   , VVM_INS_PARAM_FFFFFF) /* 129 */           \
-	VVM_OP_XM( mad_sub1                     , Op        , 4    , 1   , f      , 0    , 0   , VVM_INS_PARAM_FFFFFF) /* 130 */           \
-	VVM_OP_XM( mad_mul                      , Op        , 4    , 1   , f      , 0    , 0   , VVM_INS_PARAM_FFFFFF) /* 131 */           \
-	VVM_OP_XM( mad_sqrt                     , Op        , 3    , 1   , f      , 0    , 0   , VVM_INS_PARAM_FFFFFF) /* 132 */           \
-	VVM_OP_XM( mad_mad0                     , Op        , 5    , 1   , f      , 0    , 0   , VVM_INS_PARAM_FFFFFF) /* 133 */           \
-	VVM_OP_XM( mad_mad1                     , Op        , 5    , 1   , f      , 0    , 0   , VVM_INS_PARAM_FFFFFF) /* 134 */           \
-	VVM_OP_XM( mul_mad0                     , Op        , 4    , 1   , f      , 0    , 0   , VVM_INS_PARAM_FFFFFF) /* 135 */           \
-	VVM_OP_XM( mul_mad1                     , Op        , 4    , 1   , f      , 0    , 0   , VVM_INS_PARAM_FFFFFF) /* 136 */           \
-	VVM_OP_XM( mul_add                      , Op        , 3    , 1   , f      , 0    , 0   , VVM_INS_PARAM_FFFFFF) /* 137 */           \
-	VVM_OP_XM( mul_sub0                     , Op        , 3    , 1   , f      , 0    , 0   , VVM_INS_PARAM_FFFFFF) /* 138 */           \
-	VVM_OP_XM( mul_sub1                     , Op        , 3    , 1   , f      , 0    , 0   , VVM_INS_PARAM_FFFFFF) /* 139 */           \
-	VVM_OP_XM( mul_mul                      , Op        , 3    , 1   , f      , 0    , 0   , VVM_INS_PARAM_FFFFFF) /* 140 */           \
-	VVM_OP_XM( mul_max                      , Op        , 3    , 1   , f      , 0    , 0   , VVM_INS_PARAM_FFFFFF) /* 141 */           \
-	VVM_OP_XM( mul_2x                       , Op        , 2    , 2   , f      , 0    , 0   , VVM_INS_PARAM_FFFFFF) /* 142 */           \
-	VVM_OP_XM( add_mad1                     , Op        , 4    , 1   , f      , 0    , 0   , VVM_INS_PARAM_FFFFFF) /* 143 */           \
-	VVM_OP_XM( add_add                      , Op        , 3    , 1   , f      , 0    , 0   , VVM_INS_PARAM_FFFFFF) /* 144 */           \
-	VVM_OP_XM( sub_cmplt1                   , Op        , 3    , 1   , f      , 0    , 0   , VVM_INS_PARAM_FFIFFF) /* 145 */           \
-	VVM_OP_XM( sub_neg                      , Op        , 2    , 1   , f      , 0    , 0   , VVM_INS_PARAM_FFFFFF) /* 146 */           \
-	VVM_OP_XM( sub_mul                      , Op        , 3    , 1   , f      , 0    , 0   , VVM_INS_PARAM_FFFFFF) /* 147 */           \
-	VVM_OP_XM( div_mad0                     , Op        , 4    , 1   , f      , 0    , 0   , VVM_INS_PARAM_FFFFFF) /* 148 */           \
-	VVM_OP_XM( div_f2i                      , Op        , 2    , 1   , i      , 0    , 0   , VVM_INS_PARAM_FFFFFI) /* 149 */           \
-	VVM_OP_XM( div_mul                      , Op        , 3    , 1   , f      , 0    , 0   , VVM_INS_PARAM_FFFFFF) /* 150 */           \
-	VVM_OP_XM( muli_addi                    , Op        , 3    , 1   , i      , 0    , 0   , VVM_INS_PARAM_IIIIII) /* 151 */           \
-	VVM_OP_XM( addi_bit_rshift				, Op        , 3    , 1   , i      , 0    , 0   , VVM_INS_PARAM_IIIIII) /* 152 */           \
-	VVM_OP_XM( addi_muli                    , Op        , 3    , 1   , i      , 0    , 0   , VVM_INS_PARAM_IIIIII) /* 153 */           \
-	VVM_OP_XM( b2i_2x                       , Op        , 1    , 2   , i      , 0    , 0   , VVM_INS_PARAM_IIIIII) /* 154 */           \
-	VVM_OP_XM( i2f_div0                     , Op        , 2    , 1   , f      , 0    , 0   , VVM_INS_PARAM_FFFFFF) /* 155 */           \
-	VVM_OP_XM( i2f_div1                     , Op        , 2    , 1   , f      , 0    , 0   , VVM_INS_PARAM_FFFFFF) /* 156 */           \
-	VVM_OP_XM( i2f_mul                      , Op        , 2    , 1   , f      , 0    , 0   , VVM_INS_PARAM_FFFFFF) /* 157 */           \
-	VVM_OP_XM( i2f_mad0                     , Op        , 3    , 1   , f      , 0    , 0   , VVM_INS_PARAM_FFFFFF) /* 158 */           \
-	VVM_OP_XM( i2f_mad1                     , Op        , 3    , 1   , f      , 0    , 0   , VVM_INS_PARAM_FFFFFF) /* 159 */           \
-	VVM_OP_XM( f2i_select1                  , Op        , 3    , 1   , i      , 0    , 0   , VVM_INS_PARAM_IIFIII) /* 160 */           \
-	VVM_OP_XM( f2i_maxi                     , Op        , 2    , 1   , i      , 0    , 0   , VVM_INS_PARAM_IIIFII) /* 161 */           \
-	VVM_OP_XM( f2i_addi                     , Op        , 2    , 1   , i      , 0    , 0   , VVM_INS_PARAM_IIIFII) /* 162 */           \
-	VVM_OP_XM( fmod_add                     , Op        , 3    , 1   , f      , 0    , 0   , VVM_INS_PARAM_FFFFFF) /* 163 */           \
-	VVM_OP_XM( bit_and_i2f                  , Op        , 2    , 1   , f      , 0    , 0   , VVM_INS_PARAM_FFFFFF) /* 164 */           \
-	VVM_OP_XM( bit_rshift_bit_and           , Op        , 3    , 1   , i      , 0    , 0   , VVM_INS_PARAM_FFFFFF) /* 165 */           \
-	VVM_OP_XM( neg_cmplt                    , Op        , 2    , 1   , f      , 0    , 0   , VVM_INS_PARAM_FFFFFF) /* 166 */           \
-	VVM_OP_XM( bit_or_muli                  , Op        , 3    , 1   , i      , 0    , 0   , VVM_INS_PARAM_FFFFFF) /* 167 */           \
-	VVM_OP_XM( bit_lshift_bit_or            , Op        , 3    , 1   , i      , 0    , 0   , VVM_INS_PARAM_FFFFFF) /* 168 */           \
-	VVM_OP_XM( random_add                   , Op        , 2    , 1   , null   , 0    , 0   , VVM_INS_PARAM_FFFFFF) /* 169 */           \
-	VVM_OP_XM( random_2x                    , Op        , 1    , 2   , null   , 0    , 0   , VVM_INS_PARAM_FFFFFF) /* 170 */           \
-	VVM_OP_XM( max_f2i                      , Op        , 2    , 1   , i      , 0    , 0   , VVM_INS_PARAM_FFFFFF) /* 171 */           \
-	VVM_OP_XM( select_mul                   , Op        , 4    , 1   , f      , 0    , 0   , VVM_INS_PARAM_FFFFFF) /* 172 */           \
-	VVM_OP_XM( select_add                   , Op        , 4    , 1   , f      , 0    , 0   , VVM_INS_PARAM_FFFFFF) /* 173 */           \
-	VVM_OP_XM( sin_cos                      , Op        , 1    , 2   , null   , 0    , 0   , VVM_INS_PARAM_FFFFFF) /* 174 */           
+	/* Merged ops -- combined ops that show up frequently together in Fornite.  There are three types: */                                                    \
+	/* 1. exec_index that get immediately fed into an add or i2f                                       */                                                    \
+	/* 2. ops with identical inputs                                                                    */                                                    \
+	/* 3. ops where the output of one chain to the next.. ie a mul that feeds directly into a sub      */                                                    \
+	VVM_OP_XM( exec_indexf                  , Op        , 1    , 1   , null     , 0             , 0               , VVM_INS_PARAM_FFIIIF) /* 98 */           \
+	VVM_OP_XM( exec_index_addi              , Op        , 1    , 1   , null     , 0             , 0               , VVM_INS_PARAM_FFIIII) /* 99 */           \
+	VVM_OP_XM( cmplt_select                 , Op        , 4    , 1   , f        , 0             , 0               , VVM_INS_PARAM_FFFFFF) /* 100 */           \
+	VVM_OP_XM( cmple_select                 , Op        , 4    , 1   , f        , 0             , 0               , VVM_INS_PARAM_FFFFFF) /* 101 */           \
+	VVM_OP_XM( cmpeq_select                 , Op        , 4    , 1   , f        , 0             , 0               , VVM_INS_PARAM_FFFFFF) /* 102 */          \
+	VVM_OP_XM( cmplti_select                , Op        , 4    , 1   , f        , 0             , 0               , VVM_INS_PARAM_FIIFFF) /* 103 */          \
+	VVM_OP_XM( cmplei_select                , Op        , 4    , 1   , f        , 0             , 0               , VVM_INS_PARAM_FIIFFF) /* 104 */          \
+	VVM_OP_XM( cmpeqi_select                , Op        , 4    , 1   , f        , 0             , 0               , VVM_INS_PARAM_FIIFFF) /* 105 */          \
+	VVM_OP_XM( cmplt_logic_and              , Op        , 3    , 1   , f        , 0             , 0               , VVM_INS_PARAM_FFFFII) /* 106 */          \
+	VVM_OP_XM( cmple_logic_and              , Op        , 3    , 1   , f        , 0             , 0               , VVM_INS_PARAM_FFFFII) /* 107 */          \
+	VVM_OP_XM( cmpgt_logic_and              , Op        , 3    , 1   , f        , 0             , 0               , VVM_INS_PARAM_FFFFII) /* 108 */          \
+	VVM_OP_XM( cmpge_logic_and              , Op        , 3    , 1   , f        , 0             , 0               , VVM_INS_PARAM_FFFFII) /* 109 */          \
+	VVM_OP_XM( cmpeq_logic_and              , Op        , 3    , 1   , f        , 0             , 0               , VVM_INS_PARAM_FFFFII) /* 110 */          \
+	VVM_OP_XM( cmpne_logic_and              , Op        , 3    , 1   , f        , 0             , 0               , VVM_INS_PARAM_FFFFII) /* 111 */          \
+	VVM_OP_XM( cmplti_logic_and             , Op        , 3    , 1   , i        , 0             , 0               , VVM_INS_PARAM_FIIFII) /* 112 */          \
+	VVM_OP_XM( cmplei_logic_and             , Op        , 3    , 1   , i        , 0             , 0               , VVM_INS_PARAM_FIIFII) /* 113 */          \
+	VVM_OP_XM( cmpgti_logic_and             , Op        , 3    , 1   , i        , 0             , 0               , VVM_INS_PARAM_FIIFII) /* 114 */          \
+	VVM_OP_XM( cmpgei_logic_and             , Op        , 3    , 1   , i        , 0             , 0               , VVM_INS_PARAM_FIIFIF) /* 115 */          \
+	VVM_OP_XM( cmpeqi_logic_and             , Op        , 3    , 1   , i        , 0             , 0               , VVM_INS_PARAM_FIIFIF) /* 116 */          \
+	VVM_OP_XM( cmpnei_logic_and             , Op        , 3    , 1   , i        , 0             , 0               , VVM_INS_PARAM_FIIFIF) /* 117 */          \
+	VVM_OP_XM( cmplt_logic_or               , Op        , 3    , 1   , f        , 0             , 0               , VVM_INS_PARAM_FFFFFF) /* 118 */          \
+	VVM_OP_XM( cmple_logic_or               , Op        , 3    , 1   , f        , 0             , 0               , VVM_INS_PARAM_FFFFFF) /* 119 */          \
+	VVM_OP_XM( cmpgt_logic_or               , Op        , 3    , 1   , f        , 0             , 0               , VVM_INS_PARAM_FFFFFF) /* 120 */          \
+	VVM_OP_XM( cmpge_logic_or               , Op        , 3    , 1   , f        , 0             , 0               , VVM_INS_PARAM_FFFFFF) /* 121 */          \
+	VVM_OP_XM( cmpeq_logic_or               , Op        , 3    , 1   , f        , 0             , 0               , VVM_INS_PARAM_FFFFFF) /* 122 */          \
+	VVM_OP_XM( cmpne_logic_or               , Op        , 3    , 1   , f        , 0             , 0               , VVM_INS_PARAM_FFFFFF) /* 123 */          \
+	VVM_OP_XM( cmplti_logic_or              , Op        , 3    , 1   , i        , 0             , 0               , VVM_INS_PARAM_FIIFFF) /* 124 */          \
+	VVM_OP_XM( cmplei_logic_or              , Op        , 3    , 1   , i        , 0             , 0               , VVM_INS_PARAM_FIIFFF) /* 125 */          \
+	VVM_OP_XM( cmpgti_logic_or              , Op        , 3    , 1   , i        , 0             , 0               , VVM_INS_PARAM_FIIFFF) /* 126 */          \
+	VVM_OP_XM( cmpgei_logic_or              , Op        , 3    , 1   , i        , 0             , 0               , VVM_INS_PARAM_FIIFFF) /* 127 */          \
+	VVM_OP_XM( cmpeqi_logic_or              , Op        , 3    , 1   , i        , 0             , 0               , VVM_INS_PARAM_FIIFFF) /* 128 */          \
+	VVM_OP_XM( cmpnei_logic_or              , Op        , 3    , 1   , i        , 0             , 0               , VVM_INS_PARAM_FIIFFF) /* 129 */          \
+	VVM_OP_XM( mad_add                      , Op        , 4    , 1   , f        , 0             , 0               , VVM_INS_PARAM_FFFFFF) /* 130 */          \
+	VVM_OP_XM( mad_sub0                     , Op        , 4    , 1   , f        , 0             , 0               , VVM_INS_PARAM_FFFFFF) /* 131 */          \
+	VVM_OP_XM( mad_sub1                     , Op        , 4    , 1   , f        , 0             , 0               , VVM_INS_PARAM_FFFFFF) /* 132 */          \
+	VVM_OP_XM( mad_mul                      , Op        , 4    , 1   , f        , 0             , 0               , VVM_INS_PARAM_FFFFFF) /* 133 */          \
+	VVM_OP_XM( mad_sqrt                     , Op        , 3    , 1   , f        , 0             , 0               , VVM_INS_PARAM_FFFFFF) /* 134 */          \
+	VVM_OP_XM( mad_mad0                     , Op        , 5    , 1   , f        , 0             , 0               , VVM_INS_PARAM_FFFFFF) /* 135 */          \
+	VVM_OP_XM( mad_mad1                     , Op        , 5    , 1   , f        , 0             , 0               , VVM_INS_PARAM_FFFFFF) /* 136 */          \
+	VVM_OP_XM( mul_mad0                     , Op        , 4    , 1   , f        , 0             , 0               , VVM_INS_PARAM_FFFFFF) /* 137 */          \
+	VVM_OP_XM( mul_mad1                     , Op        , 4    , 1   , f        , 0             , 0               , VVM_INS_PARAM_FFFFFF) /* 138 */          \
+	VVM_OP_XM( mul_add                      , Op        , 3    , 1   , f        , 0             , 0               , VVM_INS_PARAM_FFFFFF) /* 139 */          \
+	VVM_OP_XM( mul_sub0                     , Op        , 3    , 1   , f        , 0             , 0               , VVM_INS_PARAM_FFFFFF) /* 140 */          \
+	VVM_OP_XM( mul_sub1                     , Op        , 3    , 1   , f        , 0             , 0               , VVM_INS_PARAM_FFFFFF) /* 141 */          \
+	VVM_OP_XM( mul_mul                      , Op        , 3    , 1   , f        , 0             , 0               , VVM_INS_PARAM_FFFFFF) /* 142 */          \
+	VVM_OP_XM( mul_max                      , Op        , 3    , 1   , f        , 0             , 0               , VVM_INS_PARAM_FFFFFF) /* 143 */          \
+	VVM_OP_XM( mul_2x                       , Op        , 2    , 2   , f        , 0             , 0               , VVM_INS_PARAM_FFFFFF) /* 144 */          \
+	VVM_OP_XM( add_mad1                     , Op        , 4    , 1   , f        , 0             , 0               , VVM_INS_PARAM_FFFFFF) /* 145 */          \
+	VVM_OP_XM( add_add                      , Op        , 3    , 1   , f        , 0             , 0               , VVM_INS_PARAM_FFFFFF) /* 146 */          \
+	VVM_OP_XM( sub_cmplt1                   , Op        , 3    , 1   , f        , 0             , 0               , VVM_INS_PARAM_FFIFFF) /* 147 */          \
+	VVM_OP_XM( sub_neg                      , Op        , 2    , 1   , f        , 0             , 0               , VVM_INS_PARAM_FFFFFF) /* 148 */          \
+	VVM_OP_XM( sub_mul                      , Op        , 3    , 1   , f        , 0             , 0               , VVM_INS_PARAM_FFFFFF) /* 149 */          \
+	VVM_OP_XM( div_mad0                     , Op        , 4    , 1   , f        , 0             , 0               , VVM_INS_PARAM_FFFFFF) /* 150 */          \
+	VVM_OP_XM( div_f2i                      , Op        , 2    , 1   , i        , 0             , 0               , VVM_INS_PARAM_FFFFFI) /* 151 */          \
+	VVM_OP_XM( div_mul                      , Op        , 3    , 1   , f        , 0             , 0               , VVM_INS_PARAM_FFFFFF) /* 152 */          \
+	VVM_OP_XM( muli_addi                    , Op        , 3    , 1   , i        , 0             , 0               , VVM_INS_PARAM_IIIIII) /* 153 */          \
+	VVM_OP_XM( addi_bit_rshift				, Op        , 3    , 1   , i        , 0             , 0               , VVM_INS_PARAM_IIIIII) /* 154 */          \
+	VVM_OP_XM( addi_muli                    , Op        , 3    , 1   , i        , 0             , 0               , VVM_INS_PARAM_IIIIII) /* 155 */          \
+	VVM_OP_XM( b2i_2x                       , Op        , 1    , 2   , i        , 0             , 0               , VVM_INS_PARAM_IIIIII) /* 156 */          \
+	VVM_OP_XM( i2f_div0                     , Op        , 2    , 1   , f        , 0             , 0               , VVM_INS_PARAM_FFFFFF) /* 157 */          \
+	VVM_OP_XM( i2f_div1                     , Op        , 2    , 1   , f        , 0             , 0               , VVM_INS_PARAM_FFFFFF) /* 158 */          \
+	VVM_OP_XM( i2f_mul                      , Op        , 2    , 1   , f        , 0             , 0               , VVM_INS_PARAM_FFFFFF) /* 159 */          \
+	VVM_OP_XM( i2f_mad0                     , Op        , 3    , 1   , f        , 0             , 0               , VVM_INS_PARAM_FFFFFF) /* 160 */          \
+	VVM_OP_XM( i2f_mad1                     , Op        , 3    , 1   , f        , 0             , 0               , VVM_INS_PARAM_FFFFFF) /* 161 */          \
+	VVM_OP_XM( f2i_select1                  , Op        , 3    , 1   , i        , 0             , 0               , VVM_INS_PARAM_IIFIII) /* 162 */          \
+	VVM_OP_XM( f2i_maxi                     , Op        , 2    , 1   , i        , 0             , 0               , VVM_INS_PARAM_IIIFII) /* 163 */          \
+	VVM_OP_XM( f2i_addi                     , Op        , 2    , 1   , i        , 0             , 0               , VVM_INS_PARAM_IIIFII) /* 164 */          \
+	VVM_OP_XM( fmod_add                     , Op        , 3    , 1   , f        , 0             , 0               , VVM_INS_PARAM_FFFFFF) /* 165 */          \
+	VVM_OP_XM( bit_and_i2f                  , Op        , 2    , 1   , f        , 0             , 0               , VVM_INS_PARAM_FFFFFF) /* 166 */          \
+	VVM_OP_XM( bit_rshift_bit_and           , Op        , 3    , 1   , i        , 0             , 0               , VVM_INS_PARAM_FFFFFF) /* 167 */          \
+	VVM_OP_XM( neg_cmplt                    , Op        , 2    , 1   , f        , 0             , 0               , VVM_INS_PARAM_FFFFFF) /* 168 */          \
+	VVM_OP_XM( bit_or_muli                  , Op        , 3    , 1   , i        , 0             , 0               , VVM_INS_PARAM_FFFFFF) /* 169 */          \
+	VVM_OP_XM( bit_lshift_bit_or            , Op        , 3    , 1   , i        , 0             , 0               , VVM_INS_PARAM_FFFFFF) /* 170 */          \
+	VVM_OP_XM( random_add                   , Op        , 2    , 1   , null     , 0             , 0               , VVM_INS_PARAM_FFFFFF) /* 171 */          \
+	VVM_OP_XM( random_2x                    , Op        , 1    , 2   , null     , 0             , 0               , VVM_INS_PARAM_FFFFFF) /* 172 */          \
+	VVM_OP_XM( max_f2i                      , Op        , 2    , 1   , i        , 0             , 0               , VVM_INS_PARAM_FFFFFF) /* 173 */          \
+	VVM_OP_XM( select_mul                   , Op        , 4    , 1   , f        , 0             , 0               , VVM_INS_PARAM_FFFFFF) /* 174 */          \
+	VVM_OP_XM( select_add                   , Op        , 4    , 1   , f        , 0             , 0               , VVM_INS_PARAM_FFFFFF) /* 175 */          \
+	VVM_OP_XM( sin_cos                      , Op        , 1    , 2   , null     , 0             , 0               , VVM_INS_PARAM_FFFFFF) /* 176 */          \
+	VVM_OP_XM( outputdata_float_from_half   , Output    , 0    , 0   , null     , 0             , 0               , VVM_INS_PARAM_FFFFFI) /* 177 */          \
+	VVM_OP_XM( outputdata_half_from_half    , Output    , 0    , 0   , null     , 0             , 0               , VVM_INS_PARAM_FFFFFI) /* 178 */
+
 #define VVM_OP_XM(n, ...) n,
 enum class EVectorVMOp : uint8 {
 	VVM_OP_XM_LIST
@@ -280,13 +283,17 @@ struct FDataSetMeta
 
 	int32 IDAcquireTag;
 
+#if VECTORVM_SUPPORTS_LEGACY
 	//Temporary lock we're using for thread safety when writing to the FreeIDTable.
 	//TODO: A lock free algorithm is possible here. We can create a specialized lock free list and reuse the IDTable slots for FreeIndices as Next pointers for our LFL.
 	//This would also work well on the GPU. 
 	//UE-65856 for tracking this work.
 
 	//@NOTE (smcgrath): the lock/unlock functions can go with the ne VM, we don't need them
-#if VECTORVM_SUPPORTS_LEGACY
+	//Also worth noting: I implemented something similar to the above comment's TODO regarding a lock free algorithm.
+	//It's simplicity over the traditional approach was much appreciated, however in the end I just removed the whole thing
+	//because scripts were no longer being split across multiple threads, so there's no need for any syncronization anymore...
+	//if legacy is ever removed, all these comments can along with it.
 	FCriticalSection FreeTableLock;
 	FORCEINLINE void LockFreeTable();
 	FORCEINLINE void UnlockFreeTable();
