- new VectorVM bug fixes for half float

diff --git a/Engine/Source/Runtime/VectorVM/Private/VectorVMExperimental.cpp b/Engine/Source/Runtime/VectorVM/Private/VectorVMExperimental.cpp
--- a/Engine/Source/Runtime/VectorVM/Private/VectorVMExperimental.cpp
+++ b/Engine/Source/Runtime/VectorVM/Private/VectorVMExperimental.cpp
@@ -245,7 +245,7 @@ void VVMMemSet32(void *dst, uint32 val, size_t num_vals)
 	{
 		VectorRegister4i v4 = VectorIntSet1(val);
 		char *RESTRICT ptr = (char *)dst;
-		char *RESTRICT end_ptr = ptr + num_vals * 4 - 16;
+		char *RESTRICT end_ptr = ptr + num_vals * sizeof(val) - sizeof(v4);
 		while (ptr < end_ptr) {
 			VectorIntStore(v4, ptr);
 			ptr += sizeof(v4);
@@ -284,7 +284,7 @@ void VVMMemSet16(void *dst, uint16 val, size_t num_vals)
 	else
 	{	
 		char *RESTRICT ptr     = (char *)dst;
-		char *RESTRICT end_ptr = ptr + num_vals * 4 - 16;
+		char *RESTRICT end_ptr = ptr + num_vals * sizeof(val) - sizeof(Val4);
 		while (ptr < end_ptr) {
 			VectorIntStore(Val4, ptr);
 			ptr += sizeof(Val4);
@@ -297,8 +297,7 @@ void VVMMemSet16(void *dst, uint16 val, size_t num_vals)
 #if PLATFORM_CPU_X86_FAMILY
 #define VVM_pshufb(Src, Mask) _mm_shuffle_epi8(Src, Mask)
   // Fabian's round-to-nearest-even float to half
-  // ~48 SSE2 ops for 8 output
-static FORCEINLINE void VVM_floatToHalf(void *output, float const *input)
+  static void VVM_floatToHalf(void *output, float const *input)
 {
 	static const MS_ALIGN(16) unsigned int mask_sign[4]         GCC_ALIGN(16) = { 0x80000000u, 0x80000000u, 0x80000000u, 0x80000000u };
 	static const MS_ALIGN(16)          int c_f16max[4]          GCC_ALIGN(16) = { (127 + 16) << 23, (127 + 16) << 23, (127 + 16) << 23, (127 + 16) << 23 }; // all FP32 values >=this round to +inf
@@ -342,44 +341,14 @@ static FORCEINLINE void VVM_floatToHalf(void *output, float const *input)
     __m128i joined      = _mm_or_si128(_mm_and_si128(nonspecial, b_isregular), _mm_andnot_si128(b_isregular, inf_or_nan));
 
     __m128i sign_shift  = _mm_srai_epi32(_mm_castps_si128(justsign), 16);
-    __m128i final2, final= _mm_or_si128(joined, sign_shift);
-
-    f           =  _mm_loadu_ps(input+4);
-    justsign    = _mm_and_ps(msign, f);
-    absf        = _mm_xor_ps(f, justsign);
-    absf_int    = _mm_castps_si128(absf); // the cast is "free" (extra bypass latency, but no thruput hit)
-    b_isnan     = _mm_cmpunord_ps(absf, absf); // is this a NaN?
-    b_isregular = _mm_cmpgt_epi32(f16max, absf_int); // (sub)normalized or special?
-    nanbit      = _mm_and_si128(_mm_castps_si128(b_isnan), *(VectorRegister4i *)c_nanbit);
-    inf_or_nan  = _mm_or_si128(nanbit, *(VectorRegister4i *)c_infty_as_fp16); // output for specials
-
-    b_issub     = _mm_cmpgt_epi32(min_normal, absf_int);
-
-    // "result is subnormal" path
-    subnorm1    = _mm_add_ps(absf, _mm_castsi128_ps(*(VectorRegister4i *)c_subnorm_magic)); // magic value to round output mantissa
-    subnorm2    = _mm_sub_epi32(_mm_castps_si128(subnorm1), *(VectorRegister4i *)c_subnorm_magic); // subtract out bias
-
-    // "result is normal" path
-    mantoddbit  = _mm_slli_epi32(absf_int, 31 - 13); // shift bit 13 (mantissa LSB) to sign
-    mantodd     = _mm_srai_epi32(mantoddbit, 31); // -1 if FP16 mantissa odd, else 0
-
-    round1      = _mm_add_epi32(absf_int, *(VectorRegister4i *)c_normal_bias);
-    round2      = _mm_sub_epi32(round1, mantodd); // if mantissa LSB odd, bias towards rounding up (RTNE)
-    normal      = _mm_srli_epi32(round2, 13); // rounded result
+	__m128i res         = _mm_or_si128(joined, sign_shift);
 
-    // combine the two non-specials
-    nonspecial  = _mm_or_si128(_mm_and_si128(subnorm2, b_issub), _mm_andnot_si128(b_issub, normal));
-
-    // merge in specials as well
-    joined      = _mm_or_si128(_mm_and_si128(nonspecial, b_isregular), _mm_andnot_si128(b_isregular, inf_or_nan));
-
-    sign_shift  = _mm_srai_epi32(_mm_castps_si128(justsign), 16);
-    final2      = _mm_or_si128(joined, sign_shift);
-    final       = _mm_packs_epi32(final, final2);
-	VectorIntStore(final, output);
+	res = _mm_packs_epi32(res, res);
+	_mm_storeu_si64(output, res);
 }
 
 
+
 VM_FORCEINLINE VectorRegister4i VVMIntRShift(VectorRegister4i v0, VectorRegister4i v1)
 {
 	uint32 *v0_4 = (uint32 *)&v0;
@@ -907,9 +876,9 @@ static const uint8 *VVM_Output16(const bool CT_MultipleLoops, const uint8 *InsPt
 				if (SrcInc == 0) //setting from a constant
 				{ 
 					//constants are 32 bits so convert
-					uint16 Val;
+					uint16 Val[8];
 					VVM_floatToHalf(&Val, (float *)SrcReg);
-					VVMMemSet16(DstReg, Val, NumOutputInstances);
+					VVMMemSet16(DstReg, Val[0], NumOutputInstances);
 				}
 				else
 				{
@@ -931,9 +900,9 @@ static const uint8 *VVM_Output16(const bool CT_MultipleLoops, const uint8 *InsPt
 				char *DstReg = (char *)RegPtrTable[DstIndices[j]] + (InstanceOffset * sizeof(uint16));
 				if (SrcInc == 0) //setting from a constant
 				{ 
-					uint16 Val;
+					uint16 Val[8];
 					VVM_floatToHalf(&Val, (float *)RegPtrTable[SrcIndices[j]]);
-					VVMMemSet16(DstReg, Val, NumOutputInstances);
+					VVMMemSet16(DstReg, Val[0], NumOutputInstances);
 				}
 				else
 				{
@@ -947,8 +916,9 @@ static const uint8 *VVM_Output16(const bool CT_MultipleLoops, const uint8 *InsPt
 						uint8 TblIdx = *TblIdxPtr++;
 						VectorRegister4i Mask = ((VectorRegister4i *)VVM_PSHUFB_OUTPUT_TABLE16)[TblIdx];
 						VectorRegister4i Src  = VectorIntLoad((VectorRegister4i *)SrcPtr); //loading from a temp register and they're always aligned
-						VectorRegister4i Val  = VVM_pshufb(Src, Mask);
-						VectorIntStore_16(Val, DstPtr);
+						VectorRegister4i Val;
+						FPlatformMath::VectorStoreHalf((uint16*) &Val, (float*)SrcPtr);
+						VectorIntStore_16(VVM_pshufb(Val, Mask), DstPtr);
 						SrcPtr += sizeof(VectorRegister4i);
 						DstPtr += VVM_OUTPUT_ADVANCE_TABLE16[TblIdx];
 					}
@@ -966,7 +936,7 @@ static const uint8 *VVM_Output16(const bool CT_MultipleLoops, const uint8 *InsPt
 			//convert 4 values at once then shift them in place
 			VectorRegister4i Mask = ((VectorRegister4i *)VVM_PSHUFB_OUTPUT_TABLE16)[OutputMask];
 			VectorRegister4i HalfVals;
-			VVM_floatToHalf(&HalfVals, SrcReg);
+			FPlatformMath::VectorStoreHalf((uint16*) &HalfVals, SrcReg);
 			VectorIntStore_16(VVM_pshufb(HalfVals, Mask), DstReg);
 		}
 	}
diff --git a/Engine/Source/Runtime/VectorVM/Private/VectorVMExperimentalOptimizer.cpp b/Engine/Source/Runtime/VectorVM/Private/VectorVMExperimentalOptimizer.cpp
--- a/Engine/Source/Runtime/VectorVM/Private/VectorVMExperimentalOptimizer.cpp
+++ b/Engine/Source/Runtime/VectorVM/Private/VectorVMExperimentalOptimizer.cpp
@@ -2753,7 +2753,25 @@ uint32 OptimizeVectorVMScript(const uint8 *InBytecode, int InBytecodeLen, FVecto
 				uint16 SerialIdx = OutputRegs[i].OutputIns->Output.SerialIdx;
 				check(OutputRegs[i].OutputIns->Output.DataSetIdx < 0xFF); //storing 1 byte only
 				OptContext->OutputRemapDataSetIdx[SerialIdx] = (uint8)OutputRegs[i].OutputIns->Output.DataSetIdx;
-				OptContext->OutputRemapDataType[SerialIdx]   = (uint16)OutputRegs[i].OutputIns->OpCode - (uint16)EVectorVMOp::outputdata_float;
+				switch (OutputRegs[i].OutputIns->OpCode)
+				{
+					case EVectorVMOp::outputdata_float:
+						OptContext->OutputRemapDataType[SerialIdx] = 0;
+					break;
+
+					case EVectorVMOp::outputdata_int32:
+						OptContext->OutputRemapDataType[SerialIdx] = 1;
+					break;
+
+					case EVectorVMOp::outputdata_half:
+					case EVectorVMOp::outputdata_float_from_half:
+					case EVectorVMOp::outputdata_half_from_half:
+						OptContext->OutputRemapDataType[SerialIdx] = 2;
+					break;
+
+					default:
+					break;
+				}
 				OptContext->OutputRemapDst[SerialIdx]        = OutputRegs[i].OutputIns->Output.DstRegIdx;
 			}			
 
@@ -3149,4 +3167,4 @@ uint32 OptimizeVectorVMScript(const uint8 *InBytecode, int InBytecodeLen, FVecto
 
 #undef VectorVMOptimizerSetError
 
-#endif //VECTORVM_SUPPORTS_EXPERIMENTAL
\ No newline at end of file
+#endif //VECTORVM_SUPPORTS_EXPERIMENTAL
