Fix divi for the EVVM to treat the operands as signed integers

diff --git a/Engine/Source/Runtime/VectorVM/Private/VectorVMExperimental.cpp b/Engine/Source/Runtime/VectorVM/Private/VectorVMExperimental.cpp
--- a/Engine/Source/Runtime/VectorVM/Private/VectorVMExperimental.cpp
+++ b/Engine/Source/Runtime/VectorVM/Private/VectorVMExperimental.cpp
@@ -158,12 +158,13 @@ static const uint8 VVM_OUTPUT_ADVANCE_TABLE16[] =
 	4, 6, 6, 8
 };
 
-struct FVVM_VU4
+struct FVVM_VUI4
 {
 	union
 	{
-		VectorRegister4i i4;
-		uint32 o4[4];
+		VectorRegister4i v;
+		uint32 u4[4];
+		int32 i4[4];
 	};
 };
 
@@ -354,14 +355,14 @@ VM_FORCEINLINE VectorRegister4i VVMIntRShift(VectorRegister4i v0, VectorRegister
 	uint32 *v0_4 = (uint32 *)&v0;
 	uint32 *v1_4 = (uint32 *)&v1;
 
-	FVVM_VU4 res;
+	FVVM_VUI4 res;
 
-	res.o4[0] = v0_4[0] >> v1_4[0];
-    res.o4[1] = v0_4[1] >> v1_4[1];
-    res.o4[2] = v0_4[2] >> v1_4[2];
-    res.o4[3] = v0_4[3] >> v1_4[3];
+	res.u4[0] = v0_4[0] >> v1_4[0];
+	res.u4[1] = v0_4[1] >> v1_4[1];
+	res.u4[2] = v0_4[2] >> v1_4[2];
+	res.u4[3] = v0_4[3] >> v1_4[3];
 	
-	return res.i4;
+	return res.v;
 }
 
 VM_FORCEINLINE VectorRegister4i VVMIntLShift(VectorRegister4i v0, VectorRegister4i v1)
@@ -369,14 +370,14 @@ VM_FORCEINLINE VectorRegister4i VVMIntLShift(VectorRegister4i v0, VectorRegister
 	uint32 *v0_4 = (uint32 *)&v0;
 	uint32 *v1_4 = (uint32 *)&v1;
     
-	FVVM_VU4 res;
+	FVVM_VUI4 res;
 
-	res.o4[0] = v0_4[0] << v1_4[0];
-    res.o4[1] = v0_4[1] << v1_4[1];
-    res.o4[2] = v0_4[2] << v1_4[2];
-    res.o4[3] = v0_4[3] << v1_4[3];
+	res.u4[0] = v0_4[0] << v1_4[0];
+	res.u4[1] = v0_4[1] << v1_4[1];
+	res.u4[2] = v0_4[2] << v1_4[2];
+	res.u4[3] = v0_4[3] << v1_4[3];
 	
-	return res.i4;
+	return res.v;
 }
 
 #elif PLATFORM_CPU_ARM_FAMILY
@@ -1237,17 +1238,17 @@ VM_FORCEINLINE bool VVM_serSyncRandom(const uint8 *InsPtr, FVectorVMBatchState *
 
 static VM_FORCEINLINE VectorRegister4i VVMIntDiv(VectorRegister4i v0, VectorRegister4i v1)
 {
-	uint32 *v0_4 = (uint32 *)&v0;
-	uint32 *v1_4 = (uint32 *)&v1;
+	const int32 *v0_4 = reinterpret_cast<const int32*>(&v0);
+	const int32 *v1_4 = reinterpret_cast<const int32*>(&v1);
 
-    FVVM_VU4 res;
+	FVVM_VUI4 res;
 
-	res.o4[0] = v1_4[0] == 0 ? 0 : (v0_4[0] / v1_4[0]);
-    res.o4[1] = v1_4[1] == 0 ? 0 : (v0_4[1] / v1_4[1]);
-    res.o4[2] = v1_4[2] == 0 ? 0 : (v0_4[2] / v1_4[2]);
-    res.o4[3] = v1_4[3] == 0 ? 0 : (v0_4[3] / v1_4[3]);
+	res.i4[0] = v1_4[0] == 0 ? 0 : (v0_4[0] / v1_4[0]);
+	res.i4[1] = v1_4[1] == 0 ? 0 : (v0_4[1] / v1_4[1]);
+	res.i4[2] = v1_4[2] == 0 ? 0 : (v0_4[2] / v1_4[2]);
+	res.i4[3] = v1_4[3] == 0 ? 0 : (v0_4[3] / v1_4[3]);
 	
-	return res.i4;
+	return res.v;
 }
 
 
