Fix for merged gt/ge + select ops not being properly mirrored during VM optimization pass

diff --git a/Engine/Source/Runtime/Core/Private/UObject/FortniteMainBranchObjectVersion.cpp b/Engine/Source/Runtime/Core/Private/UObject/FortniteMainBranchObjectVersion.cpp
--- a/Engine/Source/Runtime/Core/Private/UObject/FortniteMainBranchObjectVersion.cpp
+++ b/Engine/Source/Runtime/Core/Private/UObject/FortniteMainBranchObjectVersion.cpp
@@ -11,7 +11,7 @@ TMap<FGuid, FGuid> FFortniteMainBranchObjectVersion::GetSystemGuids()
 	SystemGuids.Add(DevGuids.MATERIALSHADERMAP_DERIVEDDATA_VER, FGuid("CE7776E63E9E43FBB8C103EE1AB40B5D"));
 	SystemGuids.Add(DevGuids.NANITE_DERIVEDDATA_VER, FGuid("D28C003F825F4C178F698938FB0E47F7"));
 	SystemGuids.Add(DevGuids.NIAGARASHADERMAP_DERIVEDDATA_VER, FGuid("D0DF49D408F34E38A98FA097E0280EC7"));
-	SystemGuids.Add(DevGuids.Niagara_LatestScriptCompileVersion, FGuid("2C6493A83208421484D9D9EA53B05ACE"));
+	SystemGuids.Add(DevGuids.Niagara_LatestScriptCompileVersion, FGuid("A5E27DDD78FE45DDBC0DA1CE12C6A8A1"));
 	SystemGuids.Add(DevGuids.SkeletalMeshDerivedDataVersion, FGuid("518D7C51D31D9BE4A9F4B81542AC8119"));
 	SystemGuids.Add(DevGuids.STATICMESH_DERIVEDDATA_VER, FGuid("543843B103794E0E9BA4BE80FB602F79"));
 	return SystemGuids;
diff --git a/Engine/Source/Runtime/VectorVM/Private/VectorVMExperimentalOptimizer.cpp b/Engine/Source/Runtime/VectorVM/Private/VectorVMExperimentalOptimizer.cpp
--- a/Engine/Source/Runtime/VectorVM/Private/VectorVMExperimentalOptimizer.cpp
+++ b/Engine/Source/Runtime/VectorVM/Private/VectorVMExperimentalOptimizer.cpp
@@ -1900,14 +1900,14 @@ uint32 OptimizeVectorVMScript(const uint8 *InBytecode, int InBytecodeLen, FVecto
 				{
 					case EVectorVMOp::cmplt:    NewOpCode = EVectorVMOp::cmplt_select;                             break;
 					case EVectorVMOp::cmple:    NewOpCode = EVectorVMOp::cmple_select;                             break;
-					case EVectorVMOp::cmpgt:    NewOpCode = EVectorVMOp::cmplt_select;      ReverseInputs = true;  break;
-					case EVectorVMOp::cmpge:    NewOpCode = EVectorVMOp::cmple_select;      ReverseInputs = true;  break;
+					case EVectorVMOp::cmpgt:    NewOpCode = EVectorVMOp::cmple_select;      ReverseInputs = true;  break;
+					case EVectorVMOp::cmpge:    NewOpCode = EVectorVMOp::cmplt_select;      ReverseInputs = true;  break;
 					case EVectorVMOp::cmpeq:    NewOpCode = EVectorVMOp::cmpeq_select;                             break;
 					case EVectorVMOp::cmpneq:   NewOpCode = EVectorVMOp::cmpeq_select;      ReverseInputs = true;  break;
 					case EVectorVMOp::cmplti:   NewOpCode = EVectorVMOp::cmplti_select;                            break;
 					case EVectorVMOp::cmplei:   NewOpCode = EVectorVMOp::cmplei_select;                            break;
-					case EVectorVMOp::cmpgti:   NewOpCode = EVectorVMOp::cmplti_select;     ReverseInputs = true;  break;
-					case EVectorVMOp::cmpgei:   NewOpCode = EVectorVMOp::cmplei_select;     ReverseInputs = true;  break;
+					case EVectorVMOp::cmpgti:   NewOpCode = EVectorVMOp::cmplei_select;     ReverseInputs = true;  break;
+					case EVectorVMOp::cmpgei:   NewOpCode = EVectorVMOp::cmplti_select;     ReverseInputs = true;  break;
 					case EVectorVMOp::cmpeqi:   NewOpCode = EVectorVMOp::cmpeqi_select;                            break;
 					case EVectorVMOp::cmpneqi:  NewOpCode = EVectorVMOp::cmpeqi_select;     ReverseInputs = true;  break;
 				}
