Adding some more checks to investigate apparent crashes with the EVVM

diff --git a/Engine/Plugins/FX/Niagara/Source/Niagara/Classes/NiagaraDataSet.h b/Engine/Plugins/FX/Niagara/Source/Niagara/Classes/NiagaraDataSet.h
--- a/Engine/Plugins/FX/Niagara/Source/Niagara/Classes/NiagaraDataSet.h
+++ b/Engine/Plugins/FX/Niagara/Source/Niagara/Classes/NiagaraDataSet.h
@@ -100,7 +100,8 @@ public:
 	NIAGARA_API void PushCPUBuffersToGPU(const TArray<FNiagaraDataBufferRef>& SourceBuffers, bool bReleaseRef, FRHICommandList& RHICmdList, ERHIFeatureLevel::Type FeatureLevel, const TCHAR* DebugSimName);
 	NIAGARA_API void Dump(int32 StartIndex, int32 NumInstances, const FString& Label, const FName& SortParameterKey = FName())const;
 
-	FORCEINLINE TArrayView<uint8 const* RESTRICT const> GetRegisterTable() const { return TArrayView<uint8 const* RESTRICT const>(RegisterTable); }
+	FORCEINLINE TArrayView<uint8 const* RESTRICT const> ReadRegisterTable() const { return TArrayView<uint8 const* RESTRICT const>(RegisterTable); }
+	FORCEINLINE TArrayView<uint8* RESTRICT const> EditRegisterTable() const { return TArrayView<uint8* RESTRICT const>(RegisterTable); }
 	
 	typedef uint32 RegisterTypeOffsetType[3];
 	FORCEINLINE const RegisterTypeOffsetType& GetRegisterTypeOffsets() { return RegisterTypeOffsets; }
diff --git a/Engine/Plugins/FX/Niagara/Source/Niagara/Private/NiagaraScriptExecutionContext.cpp b/Engine/Plugins/FX/Niagara/Source/Niagara/Private/NiagaraScriptExecutionContext.cpp
--- a/Engine/Plugins/FX/Niagara/Source/Niagara/Private/NiagaraScriptExecutionContext.cpp
+++ b/Engine/Plugins/FX/Niagara/Source/Niagara/Private/NiagaraScriptExecutionContext.cpp
@@ -140,8 +140,8 @@ void FNiagaraScriptExecutionContextBase::BindData(int32 Index, FNiagaraDataSet&
 	DataSetInfo[Index].Init(&DataSet, Input, StartInstance, bUpdateInstanceCounts);
 
 	//Would be nice to roll this and DataSetInfo into one but currently the VM being in it's own Engine module prevents this. Possibly should move the VM into Niagara itself.
-	TArrayView<uint8 const* RESTRICT const> InputRegisters = Input ? Input->GetRegisterTable() : TArrayView<uint8 const* RESTRICT const>();
-	TArrayView<uint8 const* RESTRICT const> OutputRegisters = Output ? Output->GetRegisterTable() : TArrayView<uint8 const* RESTRICT const>();
+	FDataSetMeta::FInputRegisterView InputRegisters = Input ? Input->ReadRegisterTable() : FDataSetMeta::FInputRegisterView();
+	FDataSetMeta::FOutputRegisterView OutputRegisters = Output ? Output->EditRegisterTable() : FDataSetMeta::FOutputRegisterView();
 
 	DataSetMetaTable.SetNum(FMath::Max(DataSetMetaTable.Num(), Index + 1));
 	DataSetMetaTable[Index].Init(InputRegisters, OutputRegisters, StartInstance,
@@ -167,10 +167,10 @@ void FNiagaraScriptExecutionContextBase::BindData(int32 Index, FNiagaraDataBuffe
 	FNiagaraDataSet* DataSet = Input->GetOwner();
 	DataSetInfo[Index].Init(DataSet, Input, StartInstance, bUpdateInstanceCounts);
 
-	TArrayView<uint8 const* RESTRICT const> InputRegisters = Input->GetRegisterTable();
+	FDataSetMeta::FInputRegisterView InputRegisters = Input->ReadRegisterTable();
 
 	DataSetMetaTable.SetNum(FMath::Max(DataSetMetaTable.Num(), Index + 1));
-	DataSetMetaTable[Index].Init(InputRegisters, TArrayView<uint8 const* RESTRICT const>(), StartInstance, nullptr, nullptr, DataSet->GetNumFreeIDsPtr(), &DataSet->NumSpawnedIDs, DataSet->GetMaxUsedIDPtr(), DataSet->GetIDAcquireTag(), &DataSet->GetSpawnedIDsTable());
+	DataSetMetaTable[Index].Init(InputRegisters, FDataSetMeta::FOutputRegisterView(), StartInstance, nullptr, nullptr, DataSet->GetNumFreeIDsPtr(), &DataSet->NumSpawnedIDs, DataSet->GetMaxUsedIDPtr(), DataSet->GetIDAcquireTag(), &DataSet->GetSpawnedIDsTable());
 
 	if (InputRegisters.Num() > 0)
 	{
diff --git a/Engine/Source/Runtime/VectorVM/Private/VectorVMExperimental.cpp b/Engine/Source/Runtime/VectorVM/Private/VectorVMExperimental.cpp
--- a/Engine/Source/Runtime/VectorVM/Private/VectorVMExperimental.cpp
+++ b/Engine/Source/Runtime/VectorVM/Private/VectorVMExperimental.cpp
@@ -405,7 +405,7 @@ VM_FORCEINLINE VectorRegister4i VVMIntLShift(VectorRegister4i v0, VectorRegister
 
 #endif
 
-static uint8 *SetupBatchStatePtrs(FVectorVMExecContext *ExecCtx, FVectorVMBatchState *BatchState)
+static void SetupBatchStatePtrs(FVectorVMExecContext *ExecCtx, FVectorVMBatchState *BatchState)
 {
 	uint8 *BatchDataPtr        = (uint8 *) VVM_ALIGN_64((size_t)BatchState + sizeof(FVectorVMBatchState));
 	size_t NumPtrRegsInTable   = ExecCtx->VVMState->NumTempRegisters + ExecCtx->VVMState->NumConstBuffers + ExecCtx->VVMState->NumInputBuffers * 2 + ExecCtx->VVMState->NumOutputBuffers;
@@ -417,6 +417,8 @@ static uint8 *SetupBatchStatePtrs(FVectorVMExecContext *ExecCtx, FVectorVMBatchS
 	BatchState->ChunkLocalData.StartingOutputIdxPerDataSet	= (uint32 *) BatchDataPtr;                  BatchDataPtr += ExecCtx->VVMState->ChunkLocalDataOutputIdxNumBytes;
 	BatchState->ChunkLocalData.NumOutputPerDataSet          = (uint32 *) BatchDataPtr;                  BatchDataPtr += ExecCtx->VVMState->ChunkLocalNumOutputNumBytes;
 	BatchState->ChunkLocalData.OutputMaskIdx                = (uint8 **) BatchDataPtr;                  BatchDataPtr += ExecCtx->VVMState->ChunkLocalOutputMaskIdxNumBytes;
+	BatchState->ChunkLocalData.RandCounters = nullptr; //these get malloc'd separately if they're ever used... which they very rarely are
+
 	for (uint32 i = 0; i < ExecCtx->VVMState->MaxOutputDataSet; ++i) {
 		BatchState->ChunkLocalData.OutputMaskIdx[i] = BatchState->OutputMaskIdx + i * NumLoops;
 	}
@@ -432,7 +434,6 @@ static uint8 *SetupBatchStatePtrs(FVectorVMExecContext *ExecCtx, FVectorVMBatchS
 	size_t PtrAfterExtFnDecodeReg = (size_t)BatchDataPtr;
 	check(PtrAfterExtFnDecodeReg - PtrStart <= ExecCtx->VVMState->BatchOverheadSize + ExecCtx->Internal.PerBatchRegisterDataBytesRequired);
 
-
 	{ //build the register pointer table which contains pointers (in order) to:
 		uint32 **TempRegPtr     = (uint32 **)BatchState->RegPtrTable;                      //1. Temp Registers
 		uint32 **ConstBuffPtr   = TempRegPtr   + ExecCtx->VVMState->NumTempRegisters;      //2. Constant Buffers
@@ -492,25 +493,27 @@ static uint8 *SetupBatchStatePtrs(FVectorVMExecContext *ExecCtx, FVectorVMBatchS
 		//outputs
 		for (uint32 i = 0; i < ExecCtx->VVMState->NumOutputBuffers; ++i)
 		{
-			uint8 DataSetIdx              = ExecCtx->VVMState->OutputRemapDataSetIdx[i];
-			uint16 OutputDataType         = ExecCtx->VVMState->OutputRemapDataType[i];
-			uint16 OutputMapDst           = ExecCtx->VVMState->OutputRemapDst[i];
-			uint32 TypeOffset             = ExecCtx->DataSets[DataSetIdx].OutputRegisterTypeOffsets[OutputDataType];
-			uint32 InstanceOffset         = ExecCtx->DataSets[DataSetIdx].InstanceOffset;
-			uint32 **DataSetOutputBuffers = (uint32 **)ExecCtx->DataSets[DataSetIdx].OutputRegisters.GetData();
-			OutputPtr[i] = OutputMapDst == 0xFFFF ? NULL : (DataSetOutputBuffers[TypeOffset + OutputMapDst] + InstanceOffset); //I don't actually think this can be null?
-		}
-	}
+			const uint8 DataSetIdx              = ExecCtx->VVMState->OutputRemapDataSetIdx[i];
+			check(DataSetIdx < ExecCtx->DataSets.Num());
 
-	BatchState->ChunkLocalData.RandCounters = nullptr; //these get malloc'd separately if they're ever used... which they very rarely are
+			const uint16 OutputDataType         = ExecCtx->VVMState->OutputRemapDataType[i];
+			check(OutputDataType < UE_ARRAY_COUNT(FDataSetMeta::OutputRegisterTypeOffsets));
 
-	if ((size_t)(BatchDataPtr - (uint8 *)BatchState) <= ExecCtx->Internal.NumBytesRequiredPerBatch)
-	{
-		return BatchDataPtr;
-	}
-	else
-	{
-		return nullptr;
+			const uint16 OutputMapDst           = ExecCtx->VVMState->OutputRemapDst[i];
+
+			if (OutputMapDst == 0xFFFF)
+			{
+				OutputPtr[i] = nullptr;
+			}
+			else
+			{
+				const uint32 TypeOffset = ExecCtx->DataSets[DataSetIdx].OutputRegisterTypeOffsets[OutputDataType];
+				const uint32 OutputBufferIdx = TypeOffset + OutputMapDst;
+				const uint32 InstanceOffset = ExecCtx->DataSets[DataSetIdx].InstanceOffset;
+
+				OutputPtr[i] = reinterpret_cast<uint32*>(ExecCtx->DataSets[DataSetIdx].OutputRegisters[OutputBufferIdx]) + InstanceOffset;
+			}
+		}
 	}
 }
 
@@ -2690,10 +2693,6 @@ void ExecVectorVMState(FVectorVMExecContext *ExecCtx, FVectorVMSerializeState *S
 		ExecCtx->VVMState->NumInstancesExecCached = ExecCtx->NumInstances;
 	}
 
-	for (uint32 i = 0; i < ExecCtx->VVMState->MaxOutputDataSet; ++i)
-	{
-		ExecCtx->VVMState->NumOutputPerDataSet[i] = 0;
-	}
 	VVMSer_initSerializationState(ExecCtx, SerializeState, SerializeState->OptimizeCtx, SerializeState->Flags | VVMSer_OptimizedBytecode);
 
 #if VECTORVM_SUPPORTS_SERIALIZATION && !defined(VVM_SERIALIZE_NO_WRITE)
@@ -2711,11 +2710,8 @@ void ExecVectorVMState(FVectorVMExecContext *ExecCtx, FVectorVMSerializeState *S
 			return;
 		}
 		FVectorVMBatchState *BatchState = (FVectorVMBatchState *)FMemory::Malloc(ExecCtx->VVMState->BatchOverheadSize + ExecCtx->Internal.PerBatchRegisterDataBytesRequired);
-		bool RegIncSetForAVX = false;
-		if (SetupBatchStatePtrs(ExecCtx, BatchState) == nullptr)
-		{
-			return;
-		}
+		SetupBatchStatePtrs(ExecCtx, BatchState);
+
 		if (ExecCtx->VVMState->Flags & VVMFlag_HasRandInstruction)
 		{
 			SetupRandStateForBatch(BatchState);
diff --git a/Engine/Source/Runtime/VectorVM/Public/VectorVMCommon.h b/Engine/Source/Runtime/VectorVM/Public/VectorVMCommon.h
--- a/Engine/Source/Runtime/VectorVM/Public/VectorVMCommon.h
+++ b/Engine/Source/Runtime/VectorVM/Public/VectorVMCommon.h
@@ -259,8 +259,11 @@ struct FStatStackEntry
 //Some require RWBuffer like support.
 struct FDataSetMeta
 {
-	TArrayView<uint8 const* RESTRICT const> InputRegisters;
-	TArrayView<uint8 const* RESTRICT const> OutputRegisters;
+	using FInputRegisterView = TArrayView<uint8 const* RESTRICT const>;
+	using FOutputRegisterView = TArrayView<uint8* RESTRICT const>;
+
+	FInputRegisterView InputRegisters;
+	FOutputRegisterView OutputRegisters;
 
 	uint32 InputRegisterTypeOffsets[3];
 	uint32 OutputRegisterTypeOffsets[3];
@@ -315,8 +318,8 @@ struct FDataSetMeta
 
 	FORCEINLINE void Reset()
 	{
-		InputRegisters = TArrayView<uint8 const* RESTRICT const>();
-		OutputRegisters = TArrayView<uint8 const* RESTRICT const>();
+		InputRegisters = FInputRegisterView();
+		OutputRegisters = FOutputRegisterView();
 		DataSetAccessIndex = INDEX_NONE;
 		InstanceOffset = INDEX_NONE;
 		IDTable = nullptr;
@@ -327,7 +330,7 @@ struct FDataSetMeta
 		IDAcquireTag = INDEX_NONE;
 	}
 
-	FORCEINLINE void Init(const TArrayView<uint8 const* RESTRICT const>& InInputRegisters, const TArrayView<uint8 const* RESTRICT const>& InOutputRegisters, int32 InInstanceOffset, TArray<int32>* InIDTable, TArray<int32>* InFreeIDTable, int32* InNumFreeIDs, int32 *InNumSpawnedIDs, int32* InMaxUsedID, int32 InIDAcquireTag, TArray<int32>* InSpawnedIDsTable)
+	FORCEINLINE void Init(FInputRegisterView InInputRegisters, FOutputRegisterView InOutputRegisters, int32 InInstanceOffset, TArray<int32>* InIDTable, TArray<int32>* InFreeIDTable, int32* InNumFreeIDs, int32 *InNumSpawnedIDs, int32* InMaxUsedID, int32 InIDAcquireTag, TArray<int32>* InSpawnedIDsTable)
 	{
 		InputRegisters = InInputRegisters;
 		OutputRegisters = InOutputRegisters;
