- Fix buffer underrun from EVVM
- If we output < 4 instance we are writing to memory before the allocation
- Data buffers are allocated 16 byte aligned and each component is aligned to 16 bytes so we can therefore just loop until all components are written

diff --git a/Engine/Source/Runtime/VectorVM/Private/VectorVMExperimental.cpp b/Engine/Source/Runtime/VectorVM/Private/VectorVMExperimental.cpp
--- a/Engine/Source/Runtime/VectorVM/Private/VectorVMExperimental.cpp
+++ b/Engine/Source/Runtime/VectorVM/Private/VectorVMExperimental.cpp
@@ -979,14 +979,13 @@ static const uint8 *VVM_Output32(const bool CT_MultipleLoops, const uint8 *InsPt
 				if (SrcInc == 0) //setting from a constant
 				{ 
 					VectorRegister4i Val   = *(VectorRegister4i *)RegPtrTable[SrcIndices[j]];
-					char * RESTRICT DstEnd = DstReg + sizeof(uint32) * NumOutputInstances - sizeof(VectorRegister4i);
+					char * RESTRICT PtrEnd = DstReg + sizeof(uint32) * NumOutputInstances;
 					char * RESTRICT Ptr    = DstReg;
-					while (Ptr < DstEnd)
+					while (Ptr < PtrEnd)
 					{
 						VectorIntStore(Val, Ptr);
 						Ptr += sizeof(VectorRegister4i);
 					}
-					VectorIntStore(Val, DstEnd);
 				}
 				else
 				{
