More fixes for EVVM with compressed (fp16) channels.  In particular this addresses issues with handling offsetting buffers when an InstanceOffset is in effect

diff --git a/Engine/Source/Runtime/VectorVM/Private/VectorVMExperimental.cpp b/Engine/Source/Runtime/VectorVM/Private/VectorVMExperimental.cpp
--- a/Engine/Source/Runtime/VectorVM/Private/VectorVMExperimental.cpp
+++ b/Engine/Source/Runtime/VectorVM/Private/VectorVMExperimental.cpp
@@ -462,33 +462,47 @@ static void SetupBatchStatePtrs(FVectorVMExecContext *ExecCtx, FVectorVMBatchSta
 			uint32 **DataSetInputBuffers = (uint32 **)ExecCtx->DataSets[DataSetIdx].InputRegisters.GetData();
 			int32 InstanceOffset         = ExecCtx->DataSets[DataSetIdx].InstanceOffset;
 
-			if (InputMapSrcIdx & 0x8000) //this is a noadvance input.  It points to data after the constant buffers
+			const bool bNoAdvanceInput = InputMapSrcIdx & 0x8000;
+			const bool bHalfInput = InputMapSrcIdx & 0x4000;
+			InputMapSrcIdx = InputMapSrcIdx & 0x3FFF;
+
+			const int32 InputRegisterIndex = ExecCtx->VVMState->NumTempRegisters + ExecCtx->VVMState->NumConstBuffers + i;
+
+			if (bNoAdvanceInput) //this is a noadvance input.  It points to data after the constant buffers
 			{
-				InputPtr[i                                     ] = (uint32 *)(ExecCtx->VVMState->ConstantBuffers + ExecCtx->VVMState->NumConstBuffers + NoAdvCounter);
-				InputPtr[i + ExecCtx->VVMState->NumInputBuffers] = (uint32 *)(ExecCtx->VVMState->ConstantBuffers + ExecCtx->VVMState->NumConstBuffers + NoAdvCounter);
+				InputPtr[i] = (uint32 *)(ExecCtx->VVMState->ConstantBuffers + ExecCtx->VVMState->NumConstBuffers + NoAdvCounter);
 				++NoAdvCounter;
 				
-				BatchState->RegIncTable[ExecCtx->VVMState->NumTempRegisters + ExecCtx->VVMState->NumConstBuffers + i] = 0; //no advance inputs... don't advance obviously
+				BatchState->RegIncTable[InputRegisterIndex] = 0; //no advance inputs... don't advance obviously
 
-				if (InputMapSrcIdx & 0x4000) //half input (@TODO: has never been tested)
+				if (bHalfInput) //half input (@TODO: has never been tested)
 				{
-					uint16 *Ptr = (uint16 *)DataSetInputBuffers[InputMapSrcIdx & 0x3FFF] + InstanceOffset;
+					uint16 *Ptr = (uint16 *)DataSetInputBuffers[InputMapSrcIdx] + InstanceOffset;
 					float val = FPlatformMath::LoadHalf(Ptr);
 					VectorRegister4f InputVal = VectorSet1(val);
 					VectorStore(InputVal, (float *)InputPtr[i]);
 				}
 				else
 				{
-					uint32 *Ptr = (uint32 *)DataSetInputBuffers[InputMapSrcIdx & 0x3FFF] + InstanceOffset;
+					uint32 *Ptr = (uint32 *)DataSetInputBuffers[InputMapSrcIdx] + InstanceOffset;
 					VectorRegister4i InputVal4 = VectorIntSet1(*Ptr);
 					VectorIntStore(InputVal4, InputPtr[i]);
 				}
 			}
 			else //regular input, point directly to the input buffer
 			{
-				InputPtr[i                                     ] = DataSetInputBuffers[InputMapSrcIdx] + InstanceOffset;
-				InputPtr[i + ExecCtx->VVMState->NumInputBuffers] = InputPtr[i]; //second copy of the "base" ptr so each chunk can start them at their correct starting offset
+				const uint32 DataTypeStride = bHalfInput ? 2 : 4;
+				const uint32 OffsetBytes = InstanceOffset * DataTypeStride;
+
+				// Note that we don't update RegIncTable because it is handled by the op being invoked
+				// (it will assume that the register is half as appropriate)
+
+				InputPtr[i] = reinterpret_cast<uint32*>(
+					reinterpret_cast<uint8*>(DataSetInputBuffers[InputMapSrcIdx]) + OffsetBytes);
 			}
+
+			//second copy of the "base" ptr so each chunk can start them at their correct starting offset
+			InputPtr[i + ExecCtx->VVMState->NumInputBuffers] = InputPtr[i]; 
 		}
 		//outputs
 		for (uint32 i = 0; i < ExecCtx->VVMState->NumOutputBuffers; ++i)
@@ -509,9 +523,10 @@ static void SetupBatchStatePtrs(FVectorVMExecContext *ExecCtx, FVectorVMBatchSta
 			{
 				const uint32 TypeOffset = ExecCtx->DataSets[DataSetIdx].OutputRegisterTypeOffsets[OutputDataType];
 				const uint32 OutputBufferIdx = TypeOffset + OutputMapDst;
-				const uint32 InstanceOffset = ExecCtx->DataSets[DataSetIdx].InstanceOffset;
+				const uint32 DataTypeStride = OutputDataType == 2 ? 2 : 4;
+				const uint32 InstanceOffsetBytes = ExecCtx->DataSets[DataSetIdx].InstanceOffset * DataTypeStride;
 
-				OutputPtr[i] = reinterpret_cast<uint32*>(ExecCtx->DataSets[DataSetIdx].OutputRegisters[OutputBufferIdx]) + InstanceOffset;
+				OutputPtr[i] = reinterpret_cast<uint32*>(ExecCtx->DataSets[DataSetIdx].OutputRegisters[OutputBufferIdx] + InstanceOffsetBytes);
 			}
 		}
 	}
@@ -700,7 +715,6 @@ static void VVMBuildMapTableCaches(FVectorVMExecContext *ExecCtx)
 		for (int i = 0; i < NumInputDataSets; ++i)
 		{
 			uint32 **DataSetInputBuffers = (uint32 **)ExecCtx->DataSets[i].InputRegisters.GetData();
-			int32 InstanceOffset = ExecCtx->DataSets[i].InstanceOffset;
 
 			//regular inputs: float, int and half
 			for (int j = 0; j < 3; ++j)
@@ -711,7 +725,7 @@ static void VVMBuildMapTableCaches(FVectorVMExecContext *ExecCtx)
 				{
 					int RemapIdx = ExecCtx->VVMState->InputDataSetOffsets[(i << 3) + j] + k;
 					ExecCtx->VVMState->InputMapCacheIdx[InputCounter] = i;
-					ExecCtx->VVMState->InputMapCacheSrc[InputCounter] = TypeOffset + ExecCtx->VVMState->InputRemapTable[RemapIdx];
+					ExecCtx->VVMState->InputMapCacheSrc[InputCounter] = TypeOffset + ExecCtx->VVMState->InputRemapTable[RemapIdx] | (((j & 2) << 13));
 					++InputCounter;
 				}
 			}
@@ -749,7 +763,6 @@ static void VVMBuildMapTableCaches(FVectorVMExecContext *ExecCtx)
 	int NumOutputLoops        = InsPtr[0];                                                                   \
 	uint8 DataSetIdx          = InsPtr[1];                                                                   \
 	uint32 NumOutputInstances = NumOutputPerDataSet[DataSetIdx];                                             \
-	uint32 InstanceOffset     = 0;                                                                           \
 	uint32 RegTypeOffset      = DataSets[DataSetIdx].OutputRegisterTypeOffsets[RegType];                     \
 	const uint16 * RESTRICT SrcIndices       = (uint16 *)(InsPtr + 2);                                       \
 	const uint16 * RESTRICT DstIndices       = SrcIndices + NumOutputLoops;                                  \
@@ -764,7 +777,7 @@ static void VVMBuildMapTableCaches(FVectorVMExecContext *ExecCtx)
 		{                                                                                                         \
 			uint16 AbsRegIdx;                                                                                     \
 			int ThisRegType = VVMGetRegisterType(ExecCtx->VVMState, SrcIndices[j], &AbsRegIdx);                   \
-			uint32 *DstReg = (uint32 *)RegPtrTable[DstIndices[j]] + InstanceOffset;                               \
+			uint32 *DstReg = (uint32 *)RegPtrTable[DstIndices[j]];                                                \
 			VECTORVM_PRINTF("\t%c%d -> %d [0x%p]\n", VVM_RT_CHAR[ThisRegType], AbsRegIdx, DstIndices[j], DstReg); \
 		}                                                                                                         \
 		return InsPtr
@@ -792,7 +805,7 @@ static const uint8 *VVM_Output16_from_16(const bool CT_MultipleLoops, const uint
 			{
 				int      SrcInc = RegIncTable[SrcIndices[j]];
 				uint32 * SrcReg = (uint32 *)RegPtrTable[SrcIndices[j]];
-				uint32 * DstReg = (uint32 *)RegPtrTable[DstIndices[j]] + InstanceOffset;
+				uint32 * DstReg = (uint32 *)RegPtrTable[DstIndices[j]];
 				if (SrcReg != DstReg)
 				{ //temp registers can be aliased to outputs
 					if (SrcInc == 0) //setting from a constant
@@ -811,7 +824,7 @@ static const uint8 *VVM_Output16_from_16(const bool CT_MultipleLoops, const uint
 			for (int j = 0; j < NumOutputLoops; ++j)
 			{
 				int     SrcInc = RegIncTable[SrcIndices[j]];
-				uint64 *DstReg = (uint64 *)RegPtrTable[DstIndices[j]] + InstanceOffset;
+				uint64 *DstReg = (uint64 *)RegPtrTable[DstIndices[j]];
 
 				if (SrcInc == 0) //setting from a constant
 				{ 
@@ -851,7 +864,7 @@ static const uint8 *VVM_Output16_from_16(const bool CT_MultipleLoops, const uint
 		for (int j = 0; j < NumOutputLoops; ++j)
 		{
 			uint32 *SrcReg = (uint32 *)RegPtrTable[SrcIndices[j]];
-			uint32 *DstReg = (uint32 *)RegPtrTable[DstIndices[j]] + InstanceOffset;
+			uint32 *DstReg = (uint32 *)RegPtrTable[DstIndices[j]];
 			VectorRegister4i Mask = ((VectorRegister4i *)VVM_PSHUFB_OUTPUT_TABLE16)[OutputMask];
 			VectorRegister4i Src  = VectorIntLoad(SrcReg);
 			VectorRegister4i Val  = VVM_pshufb(Src, Mask);
@@ -873,7 +886,7 @@ static const uint8 *VVM_Output16(const bool CT_MultipleLoops, const uint8 *InsPt
 			{
 				int     SrcInc = RegIncTable[SrcIndices[j]];
 				uint32 * RESTRICT SrcReg = (uint32 *)RegPtrTable[SrcIndices[j]];
-				uint16 * RESTRICT DstReg = (uint16 *)RegPtrTable[DstIndices[j]] + InstanceOffset;
+				uint16 * RESTRICT DstReg = (uint16 *)RegPtrTable[DstIndices[j]];
 				check((void *)SrcReg != (void *)DstReg); //half floats can't alias outputs
 				if (SrcInc == 0) //setting from a constant
 				{ 
@@ -899,7 +912,7 @@ static const uint8 *VVM_Output16(const bool CT_MultipleLoops, const uint8 *InsPt
 			for (int j = 0; j < NumOutputLoops; ++j)
 			{
 				int   SrcInc = RegIncTable[SrcIndices[j]];
-				char *DstReg = (char *)RegPtrTable[DstIndices[j]] + (InstanceOffset * sizeof(uint16));
+				char *DstReg = (char *)RegPtrTable[DstIndices[j]];
 				if (SrcInc == 0) //setting from a constant
 				{ 
 					uint16 Val[8];
@@ -934,7 +947,7 @@ static const uint8 *VVM_Output16(const bool CT_MultipleLoops, const uint8 *InsPt
 		for (int j = 0; j < NumOutputLoops; ++j)
 		{
 			float * SrcReg = (float *) RegPtrTable[SrcIndices[j]];
-			uint16 *DstReg = (uint16 *)RegPtrTable[DstIndices[j]] + InstanceOffset;
+			uint16 *DstReg = (uint16 *)RegPtrTable[DstIndices[j]];
 			//convert 4 values at once then shift them in place
 			VectorRegister4i Mask = ((VectorRegister4i *)VVM_PSHUFB_OUTPUT_TABLE16)[OutputMask];
 			VectorRegister4i HalfVals;
@@ -956,7 +969,7 @@ static const uint8 *VVM_Output32(const bool CT_MultipleLoops, const uint8 *InsPt
 			{
 				int      SrcInc = RegIncTable[SrcIndices[j]];
 				uint32 * SrcReg = (uint32 *)RegPtrTable[SrcIndices[j]];
-				uint32 * DstReg = (uint32 *)RegPtrTable[DstIndices[j]] + InstanceOffset;
+				uint32 * DstReg = (uint32 *)RegPtrTable[DstIndices[j]];
 				if (SrcReg != DstReg)
 				{ //temp registers can be aliased to outputs
 					if (SrcInc == 0) //setting from a constant
@@ -975,7 +988,7 @@ static const uint8 *VVM_Output32(const bool CT_MultipleLoops, const uint8 *InsPt
 			for (int j = 0; j < NumOutputLoops; ++j)
 			{
 				int   SrcInc = RegIncTable[SrcIndices[j]];
-				char *DstReg = (char *)RegPtrTable[DstIndices[j]] + InstanceOffset;
+				char *DstReg = (char *)RegPtrTable[DstIndices[j]];
 				if (SrcInc == 0) //setting from a constant
 				{ 
 					VectorRegister4i Val   = *(VectorRegister4i *)RegPtrTable[SrcIndices[j]];
@@ -1015,7 +1028,7 @@ static const uint8 *VVM_Output32(const bool CT_MultipleLoops, const uint8 *InsPt
 		for (int j = 0; j < NumOutputLoops; ++j)
 		{
 			uint32 *SrcReg = (uint32 *)RegPtrTable[SrcIndices[j]];
-			uint32 *DstReg = (uint32 *)RegPtrTable[DstIndices[j]] + InstanceOffset;
+			uint32 *DstReg = (uint32 *)RegPtrTable[DstIndices[j]];
 			VectorRegister4i Mask = ((VectorRegister4i *)VVM_PSHUFB_OUTPUT_TABLE)[OutputMask];
 			VectorRegister4i Src  = VectorIntLoad(SrcReg);
 			VectorRegister4i Val  = VVM_pshufb(Src, Mask);
@@ -2538,7 +2551,9 @@ static void SetBatchPointersForCorrectChunkOffsets(FVectorVMExecContext *ExecCtx
 		{
 			if (BatchState->RegIncTable[DstOffset + i] != 0) //don't offset the no-advance inputs
 			{
-				BatchState->RegPtrTable[DstOffset + i] = BatchState->RegPtrTable[SrcOffset + i] + (uint64)((uint64)BatchState->ChunkLocalData.StartInstanceThisChunk << 2ULL);
+				const uint32 DataTypeStride = ExecCtx->VVMState->InputMapCacheSrc[i] & 0x4000 ? 2 : 4;
+				const uint32 OffsetBytes = BatchState->ChunkLocalData.StartInstanceThisChunk * DataTypeStride;
+				BatchState->RegPtrTable[DstOffset + i] = BatchState->RegPtrTable[SrcOffset + i] + OffsetBytes;
 			}
 		}
 	}
@@ -2547,7 +2562,9 @@ static void SetBatchPointersForCorrectChunkOffsets(FVectorVMExecContext *ExecCtx
 		int Offset = ExecCtx->VVMState->NumTempRegisters + ExecCtx->VVMState->NumConstBuffers + ExecCtx->VVMState->NumInputBuffers * 2;
 		for (int i = 0; i < (int)ExecCtx->VVMState->NumOutputBuffers; ++i)
 		{
-			BatchState->RegPtrTable[Offset + i] = BatchState->RegPtrTable[Offset + i] + (uint64)((uint64)BatchState->ChunkLocalData.NumOutputPerDataSet[0] << 2ULL);
+			const uint32 DataTypeStride = ExecCtx->VVMState->OutputRemapDataType[i] == 2 ? 2 : 4;
+			const uint32 OffsetBytes = BatchState->ChunkLocalData.NumOutputPerDataSet[0] * DataTypeStride;
+			BatchState->RegPtrTable[Offset + i] = BatchState->RegPtrTable[Offset + i] + OffsetBytes;
 		}
 	}
 	for (uint32 i = 0; i < ExecCtx->VVMState->MaxOutputDataSet; ++i)
