Fix EVVM implementations of external register handlers to properly handle when a register is unused

diff --git a/Engine/Source/Runtime/VectorVM/Public/VectorVM.h b/Engine/Source/Runtime/VectorVM/Public/VectorVM.h
--- a/Engine/Source/Runtime/VectorVM/Public/VectorVM.h
+++ b/Engine/Source/Runtime/VectorVM/Public/VectorVM.h
@@ -134,8 +134,7 @@ namespace VectorVM
 			if (Context.UsingExperimentalVM)
 			{
 				int32 AdvanceOffset;
-				int32 RegIdx;
-				int32* ConstPtr = (int32*)Context.Experimental.GetNextRegister(&AdvanceOffset, &RegIdx);
+				int32* ConstPtr = (int32*)Context.Experimental.GetNextRegister(&AdvanceOffset);
 				check(AdvanceOffset == 0); //must be constant
 				UserPtrIdx = *ConstPtr;
 				check(UserPtrIdx != INDEX_NONE);
@@ -154,8 +153,7 @@ namespace VectorVM
 			}
 #elif VECTORVM_SUPPORTS_EXPERIMENTAL
 			int32 AdvanceOffset;
-			int32 RegIdx;
-			int32* ConstPtr = (int32*)Context.GetNextRegister(&AdvanceOffset, &RegIdx);
+			int32* ConstPtr = (int32*)Context.GetNextRegister(&AdvanceOffset);
 			check(AdvanceOffset == 0); //must be constant
 			UserPtrIdx = *ConstPtr;
 			check(UserPtrIdx != INDEX_NONE);
@@ -189,24 +187,12 @@ namespace VectorVM
 	{
 	private:
 		/** Either byte offset into constant table or offset into register table deepening on VVM_INPUT_LOCATION_BIT */
-		int32 InputOffset;
-		const T* RESTRICT InputPtr;
-		const T* RESTRICT StartPtr;
-		int32 AdvanceOffset;
-#if VECTORVM_SUPPORTS_EXPERIMENTAL && VECTORVM_SUPPORTS_LEGACY
-		bool bIsRegister;
-#endif
+		const T* RESTRICT InputPtr = nullptr;
+		const T* RESTRICT StartPtr = nullptr;
+		int32 AdvanceOffset = 0;
 
 	public:
-		FExternalFuncInputHandler()
-			: InputOffset(INDEX_NONE)
-			, InputPtr(nullptr)
-			, StartPtr(nullptr)
-			, AdvanceOffset(0)
-#if VECTORVM_SUPPORTS_EXPERIMENTAL && VECTORVM_SUPPORTS_LEGACY
-			, bIsRegister(false)
-#endif
-		{}
+		FExternalFuncInputHandler() = default;
 
 		FORCEINLINE FExternalFuncInputHandler(FVectorVMExternalFunctionContext& Context)
 		{
@@ -218,29 +204,31 @@ namespace VectorVM
 #if VECTORVM_SUPPORTS_EXPERIMENTAL && VECTORVM_SUPPORTS_LEGACY
 			if (Context.UsingExperimentalVM)
 			{
-				InputPtr = (T*)Context.Experimental.GetNextRegister(&AdvanceOffset, &InputOffset) + Context.Experimental.PerInstanceFnInstanceIdx * AdvanceOffset;
-				bIsRegister = !!AdvanceOffset;
+				InputPtr = (T*)Context.Experimental.GetNextRegister(&AdvanceOffset);
+				InputPtr += Context.Experimental.PerInstanceFnInstanceIdx * AdvanceOffset;
 			}
 			else
 			{
-				InputOffset = Context.Legacy.DecodeU16();
-				bIsRegister = !!(InputOffset & VVM_EXT_FUNC_INPUT_LOC_BIT);
+				const int32 InputOffset = Context.Legacy.DecodeU16();
+				const bool bIsRegister = !!(InputOffset & VVM_EXT_FUNC_INPUT_LOC_BIT);
+				const int32 Offset = InputOffset & VVM_EXT_FUNC_INPUT_LOC_MASK;
 
-				const int32 Offset = GetOffset();
-				InputPtr = IsConstant() ? Context.Legacy.GetConstant<T>(Offset) : reinterpret_cast<T*>(Context.Legacy.GetTempRegister(Offset));
-				AdvanceOffset = IsConstant() ? 0 : 1;
+				AdvanceOffset = bIsRegister ? 1 : 0;
+				InputPtr = bIsRegister ? reinterpret_cast<T*>(Context.Legacy.GetTempRegister(Offset)) : Context.Legacy.GetConstant<T>(Offset);
 
 				//Hack: Offset into the buffer by the instance offset.
 				InputPtr += Context.Legacy.GetExternalFunctionInstanceOffset() * AdvanceOffset;
 			}
 #elif VECTORVM_SUPPORTS_EXPERIMENTAL
-			InputPtr = (T*)Context.GetNextRegister(&AdvanceOffset, &InputOffset) + Context.PerInstanceFnInstanceIdx * AdvanceOffset;
+			InputPtr = (T*)Context.GetNextRegister(&AdvanceOffset);
+			InputPtr += Context.PerInstanceFnInstanceIdx * AdvanceOffset;
 #elif VECTORVM_SUPPORTS_LEGACY
-			InputOffset = Context.DecodeU16();
+			const int32 InputOffset = Context.DecodeU16();
+			const bool bIsRegister = !!(InputOffset & VVM_EXT_FUNC_INPUT_LOC_BIT);
+			const int32 Offset = InputOffset & VVM_EXT_FUNC_INPUT_LOC_MASK;
 
-			const int32 Offset = GetOffset();
-			InputPtr = IsConstant() ? Context.GetConstant<T>(Offset) : reinterpret_cast<T*>(Context.GetTempRegister(Offset));
-			AdvanceOffset = IsConstant() ? 0 : 1;
+			AdvanceOffset = bIsRegister ? 1 : 0;
+			InputPtr = bIsRegister ? reinterpret_cast<T*>(Context.GetTempRegister(Offset)) : Context.GetConstant<T>(Offset);
 
 			//Hack: Offset into the buffer by the instance offset.
 			InputPtr += Context.GetExternalFunctionInstanceOffset() * AdvanceOffset;
@@ -251,17 +239,8 @@ namespace VectorVM
 			StartPtr = InputPtr;
 		}
 
-		FORCEINLINE bool IsConstant()const { return !IsRegister(); }
-#if VECTORVM_SUPPORTS_EXPERIMENTAL && VECTORVM_SUPPORTS_LEGACY
-		FORCEINLINE bool IsRegister() const { return bIsRegister; }
-#elif VECTORVM_SUPPORTS_EXPERIMENTAL
-		FORCEINLINE bool IsRegister() const { return (bool)AdvanceOffset; }
-#elif VECTORVM_SUPPORTS_LEGACY
-		FORCEINLINE bool IsRegister() const { return (InputOffset & VVM_EXT_FUNC_INPUT_LOC_BIT) != 0; }
-#else
-	#error "Not supported"
-#endif
-		FORCEINLINE int32 GetOffset()const { return InputOffset & VVM_EXT_FUNC_INPUT_LOC_MASK; }
+		FORCEINLINE bool IsConstant() const { return !IsRegister(); }
+		FORCEINLINE bool IsRegister() const { return AdvanceOffset > 0; }
 		FORCEINLINE void Reset(){ InputPtr = StartPtr; }
 
 		FORCEINLINE const T Get() { return *InputPtr; }
@@ -285,69 +264,73 @@ namespace VectorVM
 	struct FExternalFuncRegisterHandler
 	{
 	private:
-		int32 RegisterIndex;
-		int32 AdvanceOffset;
+		T* RESTRICT Register = nullptr;
+		int32 AdvanceOffset = 0;
+#if VECTORVM_SUPPORTS_LEGACY
 		T Dummy;
-		T* RESTRICT Register;
+#endif
 	public:
 #if VECTORVM_SUPPORTS_EXPERIMENTAL && VECTORVM_SUPPORTS_LEGACY
-		FORCEINLINE FExternalFuncRegisterHandler(FVectorVMExternalFunctionContext& Context)
+		FExternalFuncRegisterHandler(FVectorVMExternalFunctionContext& Context)
 		{
 			if (Context.UsingExperimentalVM)
 			{
-				Register = (T*)Context.Experimental.GetNextRegister(&AdvanceOffset, &RegisterIndex);
+				Register = (T*)Context.Experimental.GetNextRegister(&AdvanceOffset);
+
+				//Hack: Offset into the buffer by the instance offset.
 				Register += Context.Experimental.PerInstanceFnInstanceIdx * AdvanceOffset;
 			}
 			else
 			{
-				RegisterIndex = (Context.Legacy.DecodeU16() & VVM_EXT_FUNC_INPUT_LOC_MASK);
-				AdvanceOffset = (IsValid() ? 1 : 0);
+				const int32 RegisterIndex = Context.Legacy.DecodeU16() & VVM_EXT_FUNC_INPUT_LOC_MASK;
+				if (RegisterIndex == static_cast<int32>(VVM_EXT_FUNC_INPUT_LOC_MASK))
 				{
-					if (IsValid())
-					{
-						checkSlow(RegisterIndex < Context.Legacy.GetNumTempRegisters());
-						Register = (T*)Context.Legacy.GetTempRegister(RegisterIndex);
-
-						//Hack: Offset into the buffer by the instance offset.
-						Register += Context.Legacy.GetExternalFunctionInstanceOffset() * AdvanceOffset;
-					}
-					else
-					{
-						Register = &Dummy;
-					}
+					Register = &Dummy;
+					AdvanceOffset = 0;
+				}
+				else
+				{
+					checkSlow(RegisterIndex < Context.Legacy.GetNumTempRegisters());
+					Register = (T*)Context.Legacy.GetTempRegister(RegisterIndex);
+					AdvanceOffset = 1;
+
+					//Hack: Offset into the buffer by the instance offset.
+					Register += Context.Legacy.GetExternalFunctionInstanceOffset() * AdvanceOffset;
 				}
 			}
 		}
 #elif VECTORVM_SUPPORTS_EXPERIMENTAL
-		FORCEINLINE FExternalFuncRegisterHandler(FVectorVMExternalFunctionContext& Context)
+		FExternalFuncRegisterHandler(FVectorVMExternalFunctionContext& Context)
 		{
-			Register = (T*)Context.GetNextRegister(&AdvanceOffset, &RegisterIndex);
+			Register = (T*)Context.GetNextRegister(&AdvanceOffset);
+
+			//Hack: Offset into the buffer by the instance offset.
 			Register += Context.PerInstanceFnInstanceIdx * AdvanceOffset;
 		}
 #elif VECTORVM_SUPPORTS_LEGACY
-		FORCEINLINE FExternalFuncRegisterHandler(FVectorVMExternalFunctionContext& Context)
-			: RegisterIndex(Context.DecodeU16() & VVM_EXT_FUNC_INPUT_LOC_MASK)
-			, AdvanceOffset(IsValid() ? 1 : 0)
+		FExternalFuncRegisterHandler(FVectorVMExternalFunctionContext& Context)
 		{
-			if (IsValid())
+			const int32 RegisterIndex = Context.DecodeU16() & VVM_EXT_FUNC_INPUT_LOC_MASK;
+			if (RegisterIndex == static_cast<int32>(VVM_EXT_FUNC_INPUT_LOC_MASK))
+			{
+				Register = &Dummy;
+				AdvanceOffset = 0;
+			}
+			else
 			{
 				checkSlow(RegisterIndex < Context.GetNumTempRegisters());
 				Register = (T*)Context.GetTempRegister(RegisterIndex);
+				AdvanceOffset = 1;
 
 				//Hack: Offset into the buffer by the instance offset.
 				Register += Context.GetExternalFunctionInstanceOffset() * AdvanceOffset;
 			}
-			else
-			{
-				Register = &Dummy;
-			}
 		}
 #else
 	#error "Not supported"
 #endif
 
-		FORCEINLINE bool IsValid() const { return RegisterIndex != (uint16)VVM_EXT_FUNC_INPUT_LOC_MASK; }
-
+		FORCEINLINE bool IsValid() const { return AdvanceOffset > 0; }
 		FORCEINLINE const T Get() { return *Register; }
 		FORCEINLINE T* GetDest() { return Register; }
 		FORCEINLINE void Advance() { Register += AdvanceOffset; }
@@ -376,9 +359,8 @@ namespace VectorVM
 		{
 			if (Context.UsingExperimentalVM)
 			{
-				int32 RegisterIndex;
 				int32 AdvanceOffset;
-				const T* Register = (const T*)Context.Experimental.GetNextRegister(&AdvanceOffset, &RegisterIndex);
+				const T* Register = (const T*)Context.Experimental.GetNextRegister(&AdvanceOffset);
 				Register += Context.Experimental.PerInstanceFnInstanceIdx * AdvanceOffset;
 
 				Constant = *Register;
@@ -392,9 +374,8 @@ namespace VectorVM
 #elif VECTORVM_SUPPORTS_EXPERIMENTAL
 		FExternalFuncConstHandler(FVectorVMExternalFunctionContext& Context)
 		{
-			int32 RegisterIndex;
 			int32 AdvanceOffset;
-			const T* Register = (const T*)Context.GetNextRegister(&AdvanceOffset, &RegisterIndex);
+			const T* Register = (const T*)Context.GetNextRegister(&AdvanceOffset);
 			Register += Context.PerInstanceFnInstanceIdx * AdvanceOffset;
 
 			Constant = *Register;
diff --git a/Engine/Source/Runtime/VectorVM/Public/VectorVMExperimental.h b/Engine/Source/Runtime/VectorVM/Public/VectorVMExperimental.h
--- a/Engine/Source/Runtime/VectorVM/Public/VectorVMExperimental.h
+++ b/Engine/Source/Runtime/VectorVM/Public/VectorVMExperimental.h
@@ -424,7 +424,7 @@ public:
 	FORCEINLINE void *                                 GetUserPtrTable(int32 UserPtrIdx) { check(UserPtrIdx < NumUserPtrs);  return UserPtrTable[UserPtrIdx]; }
 	template<uint32 InstancesPerOp> FORCEINLINE int32  GetNumLoops() const               { static_assert(InstancesPerOp == 4); return NumLoops; };
 
-	FORCEINLINE float* GetNextRegister(int32* OutAdvanceOffset, int32* OutVecIndex)
+	FORCEINLINE float* GetNextRegister(int32* OutAdvanceOffset)
 	{
 		check(RegReadCount < NumRegisters);
 		*OutAdvanceOffset = RegInc[RegReadCount];
