Change serialization of EVVM runtime data to reduce fixed overhead of the script payload by about 200 bytes/script.  Also moves some costly functions into editor only and enables (or disable the disable) of the EVVM

diff --git a/Engine/Plugins/FX/Niagara/Source/Niagara/Private/NiagaraScript.cpp b/Engine/Plugins/FX/Niagara/Source/Niagara/Private/NiagaraScript.cpp
--- a/Engine/Plugins/FX/Niagara/Source/Niagara/Private/NiagaraScript.cpp
+++ b/Engine/Plugins/FX/Niagara/Source/Niagara/Private/NiagaraScript.cpp
@@ -128,7 +128,7 @@ static FAutoConsoleVariableRef CVarNiagaraCompileHashAllDataInterfaces(
 	ECVF_Default
 );
 
-static int32 GNiagaraStripByteCodeOverride = -1; // disable EVVM for now
+static int32 GNiagaraStripByteCodeOverride = 0;
 static FAutoConsoleVariableRef CVarNiagaraStripByteCodeOverride(
 	TEXT("fx.Niagara.StripByteCodeOverride"),
 	GNiagaraStripByteCodeOverride,
diff --git a/Engine/Source/Runtime/Core/Private/UObject/FortniteReleaseBranchCustomObjectVersion.cpp b/Engine/Source/Runtime/Core/Private/UObject/FortniteReleaseBranchCustomObjectVersion.cpp
--- a/Engine/Source/Runtime/Core/Private/UObject/FortniteReleaseBranchCustomObjectVersion.cpp
+++ b/Engine/Source/Runtime/Core/Private/UObject/FortniteReleaseBranchCustomObjectVersion.cpp
@@ -6,7 +6,7 @@ TMap<FGuid, FGuid> FFortniteReleaseBranchCustomObjectVersion::GetSystemGuids()
 	TMap<FGuid, FGuid> SystemGuids;
 	const FDevSystemGuids& DevGuids = FDevSystemGuids::Get();
 
-	SystemGuids.Add(DevGuids.Niagara_LatestScriptCompileVersion, FGuid("FB079EFC568B49FBA6ECB6C4DFDD7EBB"));
+	SystemGuids.Add(DevGuids.Niagara_LatestScriptCompileVersion, FGuid("5A6F526F5B044512837BF24A20D934A4"));
 	SystemGuids.Add(DevGuids.SkeletalMeshDerivedDataVersion, FGuid("5705956EC7134274A5A1BD3FC74DB3A9"));
 	return SystemGuids;
 }
diff --git a/Engine/Source/Runtime/VectorVM/Private/VectorVMExperimental.cpp b/Engine/Source/Runtime/VectorVM/Private/VectorVMExperimental.cpp
--- a/Engine/Source/Runtime/VectorVM/Private/VectorVMExperimental.cpp
+++ b/Engine/Source/Runtime/VectorVM/Private/VectorVMExperimental.cpp
@@ -580,10 +580,18 @@ static VM_FORCEINLINE VectorRegister4i VVMXorwowStep(FVectorVMBatchState *BatchS
 	return Result;
 }
 
-FVectorVMState *AllocVectorVMState(FVectorVMOptimizeContext *OptimizeCtx) {
-	if (OptimizeCtx == nullptr || OptimizeCtx->Error.Line != 0) {
+FVectorVMState *AllocVectorVMState(FVectorVMOptimizeContext *OptimizeCtx)
+{
+	if (OptimizeCtx == nullptr)
+	{
+		return nullptr;
+	}
+#if WITH_EDITORONLY_DATA
+	if (OptimizeCtx->Error.Line != 0)
+	{
 		return nullptr;
 	}
+#endif
 	//compute the number of overhead bytes for this VVM State
 	uint32 ConstBufferOffset         = VVM_ALIGN_32(sizeof(FVectorVMState));
 
diff --git a/Engine/Source/Runtime/VectorVM/Private/VectorVMExperimentalOptimizer.cpp b/Engine/Source/Runtime/VectorVM/Private/VectorVMExperimentalOptimizer.cpp
--- a/Engine/Source/Runtime/VectorVM/Private/VectorVMExperimentalOptimizer.cpp
+++ b/Engine/Source/Runtime/VectorVM/Private/VectorVMExperimentalOptimizer.cpp
@@ -1,10 +1,207 @@
 // Copyright Epic Games, Inc. All Rights Reserved.
 #include "CoreTypes.h"
+#include "Serialization/MemoryReader.h"
+#include "Serialization/MemoryWriter.h"
 #include "VectorVM.h"
 #include "VectorVMExperimental.h"
 
+
 #if VECTORVM_SUPPORTS_EXPERIMENTAL
 
+namespace VectorVMSerializationHelper
+{
+
+struct FRuntimeContextData
+{
+	FRuntimeContextData(const FVectorVMOptimizeContext& Context)
+	: NumBytecodeBytes(Context.NumBytecodeBytes)
+	, MaxOutputDataSet(Context.MaxOutputDataSet)
+	, NumConstsAlloced(Context.NumConstsAlloced)
+	, NumTempRegisters(Context.NumTempRegisters)
+	, NumConstsRemapped(Context.NumConstsRemapped)
+	, NumInputsRemapped(Context.NumInputsRemapped)
+	, NumNoAdvanceInputs(Context.NumNoAdvanceInputs)
+	, NumInputDataSets(Context.NumInputDataSets)
+	, NumOutputsRemapped(Context.NumOutputsRemapped)
+	, NumOutputInstructions(Context.NumOutputInstructions)
+	, NumExtFns(Context.NumExtFns)
+	, MaxExtFnRegisters(Context.MaxExtFnRegisters)
+	, NumDummyRegsReq(Context.NumDummyRegsReq)
+	, MaxExtFnUsed(Context.MaxExtFnUsed)
+	, Flags(Context.Flags)
+	, HashId(Context.HashId)
+	{}
+
+	FRuntimeContextData(TConstArrayView<uint8> ContextData)
+	{
+		FMemoryReaderView Ar(ContextData);
+		Ar << *this;
+	}
+
+	void CopyToContext(FVectorVMOptimizeContext& Context) const
+	{
+		Context.NumBytecodeBytes = NumBytecodeBytes;
+		Context.MaxOutputDataSet = MaxOutputDataSet;
+		Context.NumConstsAlloced = NumConstsAlloced;
+		Context.NumTempRegisters = NumTempRegisters;
+		Context.NumConstsRemapped = NumConstsRemapped;
+		Context.NumInputsRemapped = NumInputsRemapped;
+		Context.NumNoAdvanceInputs = NumNoAdvanceInputs;
+		Context.NumInputDataSets = NumInputDataSets;
+		Context.NumOutputsRemapped = NumOutputsRemapped;
+		Context.NumOutputInstructions = NumOutputInstructions;
+		Context.NumExtFns = NumExtFns;
+		Context.MaxExtFnRegisters = MaxExtFnRegisters;
+		Context.NumDummyRegsReq = NumDummyRegsReq;
+		Context.MaxExtFnUsed = MaxExtFnUsed;
+		Context.Flags = Flags;
+		Context.HashId = HashId;
+	}
+
+	friend FArchive& operator<<(FArchive& Ar, FRuntimeContextData& ContextInfo);
+
+	uint32 NumBytecodeBytes;
+	uint32 MaxOutputDataSet;
+	uint16 NumConstsAlloced;
+	uint32 NumTempRegisters;
+	uint16 NumConstsRemapped;
+	uint16 NumInputsRemapped;
+	uint16 NumNoAdvanceInputs;
+	uint16 NumInputDataSets;
+	uint16 NumOutputsRemapped;
+	uint16 NumOutputInstructions;
+	uint32 NumExtFns;
+	uint32 MaxExtFnRegisters;
+	uint32 NumDummyRegsReq;
+	int32 MaxExtFnUsed;
+	uint32 Flags;
+	uint64 HashId;
+};
+
+FArchive& operator<<(FArchive& Ar, FRuntimeContextData& ContextInfo)
+{
+	Ar << ContextInfo.NumBytecodeBytes;
+	Ar << ContextInfo.MaxOutputDataSet;
+	Ar << ContextInfo.NumConstsAlloced;
+	Ar << ContextInfo.NumTempRegisters;
+	Ar << ContextInfo.NumConstsRemapped;
+	Ar << ContextInfo.NumInputsRemapped;
+	Ar << ContextInfo.NumNoAdvanceInputs;
+	Ar << ContextInfo.NumInputDataSets;
+	Ar << ContextInfo.NumOutputsRemapped;
+	Ar << ContextInfo.NumOutputInstructions;
+	Ar << ContextInfo.NumExtFns;
+	Ar << ContextInfo.MaxExtFnRegisters;
+	Ar << ContextInfo.NumDummyRegsReq;
+	Ar << ContextInfo.MaxExtFnUsed;
+	Ar << ContextInfo.Flags;
+	Ar << ContextInfo.HashId;
+
+	return Ar;
+}
+
+struct FContextInfoLayout
+{
+	FContextInfoLayout(const FRuntimeContextData& Context)
+		: PropertySize(sizeof(Context))
+		, BytecodeSize(Context.NumBytecodeBytes)
+		, ConstRemapSize(Context.NumConstsRemapped * sizeof(uint16))
+		, InputRemapSize(Context.NumInputsRemapped * sizeof(uint16))
+		, InputDataSetOffsetsSize(Context.NumInputDataSets * 8 * sizeof(uint16))
+		, OutputRemapDataSetIdxSize(Context.NumOutputsRemapped * sizeof(uint8))
+		, OutputRemapDataTypeSize(Context.NumOutputsRemapped * sizeof(uint16))
+		, OutputRemapDstSize(Context.NumOutputsRemapped * sizeof(uint16))
+		, ExtFnSize(Context.NumExtFns * sizeof(FVectorVMExtFunctionData))
+		, PropertyOffset(0)
+		, BytecodeOffset(Align(PropertyOffset + PropertySize, 16))
+		, ConstRemapOffset(Align(BytecodeOffset + BytecodeSize, 16))
+		, InputRemapOffset(Align(ConstRemapOffset + ConstRemapSize, 16))
+		, InputDataSetOffsetsOffset(Align(InputRemapOffset + InputRemapSize, 16))
+		, OutputRemapDataSetIdxOffset(Align(InputDataSetOffsetsOffset + InputDataSetOffsetsSize, 16))
+		, OutputRemapDataTypeOffset(Align(OutputRemapDataSetIdxOffset + OutputRemapDataSetIdxSize, 16))
+		, OutputRemapDstOffset(Align(OutputRemapDataTypeOffset + OutputRemapDataTypeSize, 16))
+		, ExtFnOffset(Align(OutputRemapDstOffset + OutputRemapDstSize, 16))
+		, TotalSize(Align(ExtFnOffset + ExtFnSize, 16))
+	{
+	}
+
+	const uint32 PropertySize;
+	const uint32 BytecodeSize;
+	const uint32 ConstRemapSize;
+	const uint32 InputRemapSize;
+	const uint32 InputDataSetOffsetsSize;
+	const uint32 OutputRemapDataSetIdxSize;
+	const uint32 OutputRemapDataTypeSize;
+	const uint32 OutputRemapDstSize;
+	const uint32 ExtFnSize;
+
+	const size_t PropertyOffset;
+	const size_t BytecodeOffset;
+	const size_t ConstRemapOffset;
+	const size_t InputRemapOffset;
+	const size_t InputDataSetOffsetsOffset;
+	const size_t OutputRemapDataSetIdxOffset;
+	const size_t OutputRemapDataTypeOffset;
+	const size_t OutputRemapDstOffset;
+	const size_t ExtFnOffset;
+	const size_t TotalSize;
+};
+
+#if WITH_EDITORONLY_DATA
+void FreezeContext(const FVectorVMOptimizeContext& Context, TArray<uint8>& ContextData)
+{
+	FRuntimeContextData RuntimeData(Context);
+	FContextInfoLayout Layout(RuntimeData);
+
+	ContextData.SetNumZeroed(Layout.TotalSize);
+	uint8* BufferData = ContextData.GetData();
+
+	FMemoryWriter Ar(ContextData);
+	Ar << RuntimeData;
+
+	FMemory::Memcpy(BufferData + Layout.BytecodeOffset, Context.OutputBytecode, Layout.BytecodeSize);
+	FMemory::Memcpy(BufferData + Layout.ConstRemapOffset, Context.ConstRemap[1], Layout.ConstRemapSize);
+	FMemory::Memcpy(BufferData + Layout.InputRemapOffset, Context.InputRemapTable, Layout.InputRemapSize);
+	FMemory::Memcpy(BufferData + Layout.InputDataSetOffsetsOffset, Context.InputDataSetOffsets, Layout.InputDataSetOffsetsSize);
+	FMemory::Memcpy(BufferData + Layout.OutputRemapDataSetIdxOffset, Context.OutputRemapDataSetIdx, Layout.OutputRemapDataSetIdxSize);
+	FMemory::Memcpy(BufferData + Layout.OutputRemapDataTypeOffset, Context.OutputRemapDataType, Layout.OutputRemapDataTypeSize);
+	FMemory::Memcpy(BufferData + Layout.OutputRemapDstOffset, Context.OutputRemapDst, Layout.OutputRemapDstSize);
+	FMemory::Memcpy(BufferData + Layout.ExtFnOffset, Context.ExtFnTable, Layout.ExtFnSize);
+
+	// external function pointers
+	FVectorVMExtFunctionData* ExtFunctionTable = reinterpret_cast<FVectorVMExtFunctionData*>(BufferData + Layout.ExtFnOffset);
+	for (uint32 ExtFunctionIt = 0; ExtFunctionIt < Context.NumExtFns; ++ExtFunctionIt)
+	{
+		ExtFunctionTable[ExtFunctionIt].Function = nullptr;
+	}
+
+}
+#endif // WITH_EDITORONLY_DATA
+
+void ThawContext(TConstArrayView<uint8> ContextData, FVectorVMOptimizeContext& Context)
+{
+	FMemory::Memzero(Context);
+
+	FRuntimeContextData RuntimeData(ContextData);
+
+	VectorVMSerializationHelper::FContextInfoLayout Layout(RuntimeData);
+
+	RuntimeData.CopyToContext(Context);
+	const uint8* BufferData = ContextData.GetData();
+	Context.OutputBytecode = const_cast<uint8*>(reinterpret_cast<const uint8*>(BufferData + Layout.BytecodeOffset));
+	Context.ConstRemap[1] = const_cast<uint16*>(reinterpret_cast<const uint16*>(BufferData + Layout.ConstRemapOffset));
+	Context.InputRemapTable = const_cast<uint16*>(reinterpret_cast<const uint16*>(BufferData + Layout.InputRemapOffset));
+	Context.InputDataSetOffsets = const_cast<uint16*>(reinterpret_cast<const uint16*>(BufferData + Layout.InputDataSetOffsetsOffset));
+	Context.OutputRemapDataSetIdx = const_cast<uint8*>(reinterpret_cast<const uint8*>(BufferData + Layout.OutputRemapDataSetIdxOffset));
+	Context.OutputRemapDataType = const_cast<uint16*>(reinterpret_cast<const uint16*>(BufferData + Layout.OutputRemapDataTypeOffset));
+	Context.OutputRemapDst = const_cast<uint16*>(reinterpret_cast<const uint16*>(BufferData + Layout.OutputRemapDstOffset));
+	Context.ExtFnTable = const_cast<FVectorVMExtFunctionData*>(reinterpret_cast<const FVectorVMExtFunctionData*>(BufferData + Layout.ExtFnOffset));
+}
+
+};
+
+#if WITH_EDITORONLY_DATA
+
 void *VVMDefaultRealloc(void *Ptr, size_t NumBytes, const char *Filename, int LineNumber);
 void VVMDefaultFree(void *Ptr, const char *Filename, int LineNumber);
 
@@ -74,120 +271,6 @@ static void VectorVMFreeOptimizerIntermediateData(FVectorVMOptimizeContext *OptC
 	}
 }
 
-void FreezeVectorVMOptimizeContext(const FVectorVMOptimizeContext& Context, TArray<uint8>& ContextData)
-{
-	const uint32 BytecodeSize                = Context.NumBytecodeBytes + 16;
-	const uint32 ConstRemapSize              = Context.NumConstsRemapped  * sizeof(uint16);
-	const uint32 InputRemapSize              = Context.NumInputsRemapped  * sizeof(uint16);
-	const uint32 InputDataSetOffsetsSize     = Context.NumInputDataSets   * 8 * sizeof(uint16);
-	const uint32 OutputRemapDataSetIdxSize   = Context.NumOutputsRemapped * sizeof(uint8);
-	const uint32 OutputRemapDataTypeSize     = Context.NumOutputsRemapped * sizeof(uint16);
-	const uint32 OutputRemapDstSize          = Context.NumOutputsRemapped * sizeof(uint16);
-	const uint32 ExtFnSize                   = Context.NumExtFns * sizeof(FVectorVMExtFunctionData);
-
-	const size_t BytecodeOffset              = Align(sizeof(FVectorVMOptimizeContext)                        , 16);
-	const size_t ConstRemapOffset            = Align(BytecodeOffset              + BytecodeSize              , 16);
-	const size_t InputRemapOffset            = Align(ConstRemapOffset            + ConstRemapSize            , 16);
-	const size_t InputDataSetOffsetsOffset   = Align(InputRemapOffset            + InputRemapSize            , 16);
-	const size_t OutputRemapDataSetIdxOffset = Align(InputDataSetOffsetsOffset   + InputDataSetOffsetsSize   , 16);
-	const size_t OutputRemapDataTypeOffset   = Align(OutputRemapDataSetIdxOffset + OutputRemapDataSetIdxSize , 16);
-	const size_t OutputRemapDstOffset        = Align(OutputRemapDataTypeOffset   + OutputRemapDataTypeSize   , 16);
-	const size_t ExtFnOffset                 = Align(OutputRemapDstOffset        + OutputRemapDstSize        , 16);
-	const size_t TotalSize                   = Align(ExtFnOffset                 + ExtFnSize                 , 16);
-
-	ContextData.SetNumZeroed(TotalSize);
-	uint8* DestData = ContextData.GetData();
-	FMemory::Memcpy(DestData, &Context, sizeof(Context));
-	FMemory::Memcpy(DestData + BytecodeOffset              , Context.OutputBytecode        , Context.NumBytecodeBytes);
-	FMemory::Memcpy(DestData + ConstRemapOffset            , Context.ConstRemap[1]         , ConstRemapSize);
-	FMemory::Memcpy(DestData + InputRemapOffset            , Context.InputRemapTable       , InputRemapSize);
-	FMemory::Memcpy(DestData + InputDataSetOffsetsOffset   , Context.InputDataSetOffsets   , InputDataSetOffsetsSize);
-
-	FMemory::Memcpy(DestData + OutputRemapDataSetIdxOffset , Context.OutputRemapDataSetIdx , OutputRemapDataSetIdxSize);
-	FMemory::Memcpy(DestData + OutputRemapDataTypeOffset   , Context.OutputRemapDataType   , OutputRemapDstSize);
-	FMemory::Memcpy(DestData + OutputRemapDstOffset        , Context.OutputRemapDst        , OutputRemapDstSize);
-	
-	FMemory::Memcpy(DestData + ExtFnOffset, Context.ExtFnTable, ExtFnSize);
-
-	// we want to clear out the pointers to any callbacks or pointers to the data that is being pushed past
-	// the end of the struct
-	FVectorVMOptimizeContext& DestContext = *reinterpret_cast<FVectorVMOptimizeContext*>(DestData);
-
-	// relocated buffers
-	DestContext.OutputBytecode        = nullptr;
-	DestContext.ConstRemap[0]         = nullptr;
-	DestContext.ConstRemap[1]         = nullptr;
-	DestContext.InputRemapTable       = nullptr;
-	DestContext.InputDataSetOffsets   = nullptr;
-	DestContext.OutputRemapDataSetIdx = nullptr;
-	DestContext.OutputRemapDataType   = nullptr;
-	DestContext.OutputRemapDst        = nullptr;
-	DestContext.ExtFnTable            = nullptr;
-
-	// external function pointers
-	FVectorVMExtFunctionData* ExtFunctionTable = reinterpret_cast<FVectorVMExtFunctionData*>(DestData + ExtFnOffset);
-	for (uint32 ExtFunctionIt = 0; ExtFunctionIt < Context.NumExtFns; ++ExtFunctionIt)
-	{
-		ExtFunctionTable[ExtFunctionIt].Function = nullptr;
-	}
-
-	// clear out unnecessary elements - todo - consider just moving to a raw dump of what we actually care about
-	FMemory::Memzero(DestContext.Init);
-	FMemory::Memzero(DestContext.Error);
-	FMemory::Memzero(DestContext.Intermediate);
-}
-
-static void* VectorVMFrozenRealloc(void* Ptr, size_t NumBytes, const char* Filename, int LineNum)
-{
-	check(false);
-	return nullptr;
-}
-
-static void VectorVMFrozenFree(void* Ptr, const char* Filename, int LineNum)
-{
-	check(false);
-}
-
-void ReinterpretVectorVMOptimizeContextData(TConstArrayView<uint8> ContextData, FVectorVMOptimizeContext& Context)
-{
-	const uint8* SrcData = ContextData.GetData();
-	const FVectorVMOptimizeContext& SrcContext = *reinterpret_cast<const FVectorVMOptimizeContext*>(SrcData);
-
-	const uint32 BytecodeSize                = SrcContext.NumBytecodeBytes + 16;
-	const uint32 ConstRemapSize              = SrcContext.NumConstsRemapped  * sizeof(uint16);
-	const uint32 InputRemapSize              = SrcContext.NumInputsRemapped  * sizeof(uint16);
-	const uint32 InputDataSetOffsetsSize     = SrcContext.NumInputDataSets   * 8 * sizeof(uint16);
-	const uint32 OutputRemapDataSetIdxSize   = SrcContext.NumOutputsRemapped * sizeof(int8);
-	const uint32 OutputRemapDataTypeSize     = SrcContext.NumOutputsRemapped * sizeof(uint16);
-	const uint32 OutputRemapDstSize          = SrcContext.NumOutputsRemapped * sizeof(int16);
-	const uint32 ExtFnSize                   = SrcContext.NumExtFns          * sizeof(FVectorVMExtFunctionData);
-
-	const size_t BytecodeOffset              = Align(sizeof(FVectorVMOptimizeContext)                        , 16);
-	const size_t ConstRemapOffset            = Align(BytecodeOffset              + BytecodeSize              , 16);
-	const size_t InputRemapOffset            = Align(ConstRemapOffset            + ConstRemapSize            , 16);
-	const size_t InputDataSetOffsetsOffset   = Align(InputRemapOffset            + InputRemapSize            , 16);
-	const size_t OutputRemapDataSetIdxOffset = Align(InputDataSetOffsetsOffset   + InputDataSetOffsetsSize   , 16);
-	const size_t OutputRemapDataTypeOffset   = Align(OutputRemapDataSetIdxOffset + OutputRemapDataSetIdxSize , 16);
-	const size_t OutputRemapDstOffset        = Align(OutputRemapDataTypeOffset   + OutputRemapDataTypeSize   , 16);
-	const size_t ExtFnOffset                 = Align(OutputRemapDstOffset        + OutputRemapDstSize        , 16);
-	const size_t TotalSize                   = Align(ExtFnOffset + ExtFnSize                                 , 16);
-
-	FMemory::Memcpy(&Context, SrcData, sizeof(Context));
-	Context.OutputBytecode        = const_cast<uint8*>(reinterpret_cast<const uint8*>(SrcData + BytecodeOffset));
-	Context.NumConstsAlloced      = Context.NumConstsRemapped;
-	Context.ConstRemap[0]         = NULL;
-	Context.ConstRemap[1]         = const_cast<uint16*>(reinterpret_cast<const uint16*>(                                     SrcData + ConstRemapOffset));
-	Context.InputRemapTable       = const_cast<uint16*>(reinterpret_cast<const uint16*>(                                     SrcData + InputRemapOffset));
-	Context.InputDataSetOffsets   = const_cast<uint16*>(reinterpret_cast<const uint16*>(                                     SrcData + InputDataSetOffsetsOffset));
-	Context.OutputRemapDataSetIdx = const_cast<uint8 *>(reinterpret_cast<const uint8 *>(                                     SrcData + OutputRemapDataSetIdxOffset));
-	Context.OutputRemapDataType   = const_cast<uint16*>(reinterpret_cast<const uint16*>(                                     SrcData + OutputRemapDataTypeOffset));
-	Context.OutputRemapDst        = const_cast<uint16*>(reinterpret_cast<const uint16*>(                                     SrcData + OutputRemapDstOffset));
-	Context.ExtFnTable            = const_cast<FVectorVMExtFunctionData*>(reinterpret_cast<const FVectorVMExtFunctionData*>( SrcData + ExtFnOffset));
-	Context.Init.ReallocFn        = VectorVMFrozenRealloc;
-	Context.Init.FreeFn           = VectorVMFrozenFree;
-	Context.Error.CallbackFn      = nullptr;
-}
-
 namespace VectorVMScriptStringHelper
 {
 
@@ -3181,4 +3264,16 @@ uint32 OptimizeVectorVMScript(const uint8 *InBytecode, int InBytecodeLen, FVecto
 
 #undef VectorVMOptimizerSetError
 
+void FreezeVectorVMOptimizeContext(const FVectorVMOptimizeContext& Context, TArray<uint8>& ContextData)
+{
+	VectorVMSerializationHelper::FreezeContext(Context, ContextData);
+}
+
+#endif // WITH_EDITORONLY_DATA
+
+void ReinterpretVectorVMOptimizeContextData(TConstArrayView<uint8> ContextData, FVectorVMOptimizeContext& Context)
+{
+	VectorVMSerializationHelper::ThawContext(ContextData, Context);
+}
+
 #endif //VECTORVM_SUPPORTS_EXPERIMENTAL
diff --git a/Engine/Source/Runtime/VectorVM/Public/VectorVMExperimental.h b/Engine/Source/Runtime/VectorVM/Public/VectorVMExperimental.h
--- a/Engine/Source/Runtime/VectorVM/Public/VectorVMExperimental.h
+++ b/Engine/Source/Runtime/VectorVM/Public/VectorVMExperimental.h
@@ -235,7 +235,6 @@ struct FVectorVMOptimizeContext
 	uint32                                NumTempRegisters;
 	uint16                                NumConstsRemapped;
 	uint16                                NumInputsRemapped;
-	uint16                                NumOutputsAliasedToTempRegisters; //@TODO: REMOVE THIS
 	uint16                                NumNoAdvanceInputs;
 	uint16                                NumInputDataSets;
 	uint16                                NumOutputsRemapped;
@@ -247,6 +246,7 @@ struct FVectorVMOptimizeContext
 	uint32                                Flags;
 	uint64                                HashId;
 
+#if WITH_EDITORONLY_DATA
 	struct
 	{
 		VectorVMReallocFn *               ReallocFn;
@@ -272,6 +272,7 @@ struct FVectorVMOptimizeContext
 		uint32                            NumInstructionsAlloced;
 		uint32                            NumRegistersUsed;
 	} Intermediate;                                           //these are freed and NULL after optimize() unless SaveIntermediateData is true when calling OptimizeVectorVMScript
+#endif
 };
 
 enum EVectorVMStateError
@@ -439,12 +440,14 @@ VECTORVM_API FVectorVMState * AllocVectorVMState                     (FVectorVMO
 VECTORVM_API void             FreeVectorVMState                      (FVectorVMState *VectorVMState);
 VECTORVM_API void             ExecVectorVMState                      (FVectorVMExecContext *ExecCtx, FVectorVMSerializeState *SerializeState, FVectorVMSerializeState *CmpSerializeState);
 
+#if WITH_EDITORONLY_DATA
 //optimize functions
 VECTORVM_API uint32           OptimizeVectorVMScript                 (const uint8 *Bytecode, int BytecodeLen, FVectorVMExtFunctionData *ExtFnIOData, int NumExtFns, FVectorVMOptimizeContext *OptContext, uint64 HashId, uint32 Flags); //OutContext must be zeroed except the Init struct
 VECTORVM_API void             FreeVectorVMOptimizeContext            (FVectorVMOptimizeContext *Context);
 VECTORVM_API void             FreezeVectorVMOptimizeContext          (const FVectorVMOptimizeContext& Context, TArray<uint8>& ContextData);
-VECTORVM_API void             ReinterpretVectorVMOptimizeContextData (TConstArrayView<uint8> ContextData, FVectorVMOptimizeContext& Context);
 VECTORVM_API void             GenerateHumanReadableVectorVMScript    (const FVectorVMOptimizeContext& Context, FString& VMScript);
+#endif
+VECTORVM_API void             ReinterpretVectorVMOptimizeContextData (TConstArrayView<uint8> ContextData, FVectorVMOptimizeContext& Context);
 
 //serialize functions
 VECTORVM_API uint32           SerializeVectorVMInputDataSets         (FVectorVMSerializeState *SerializeState, FVectorVMExecContext *ExecContext);
