- Fix for incorrect step function inside the EVVM
- A & B need to be swapped

#tests Confirm that this was a regression by testing step function on VVM vs EVVM

diff --git a/Engine/Source/Runtime/VectorVM/Private/VectorVMExperimental.cpp b/Engine/Source/Runtime/VectorVM/Private/VectorVMExperimental.cpp
--- a/Engine/Source/Runtime/VectorVM/Private/VectorVMExperimental.cpp
+++ b/Engine/Source/Runtime/VectorVM/Private/VectorVMExperimental.cpp
@@ -2373,7 +2373,7 @@ VM_FORCEINLINE VectorRegister4f VVM_Exec2f_max                             (FVec
 VM_FORCEINLINE VectorRegister4f VVM_Exec2f_pow                             (FVectorVMBatchState *BatchState, VectorRegister4f a, VectorRegister4f b)                                                            { return VectorSelect(VectorCompareGT(a, VVM_m128Const(Epsilon)), VectorPow(a, b), VectorZeroFloat()); }
 VM_FORCEINLINE VectorRegister4f VVM_Exec1f_round                           (FVectorVMBatchState *BatchState, VectorRegister4f a)                                                                                { return VectorRound(a); }
 VM_FORCEINLINE VectorRegister4f VVM_Exec1f_sign                            (FVectorVMBatchState *BatchState, VectorRegister4f a)                                                                                { return VectorSign(a); }
-VM_FORCEINLINE VectorRegister4f VVM_Exec2f_step                            (FVectorVMBatchState *BatchState, VectorRegister4f a, VectorRegister4f b)                                                            { return VectorStep(VectorSubtract(a, b)); }
+VM_FORCEINLINE VectorRegister4f VVM_Exec2f_step                            (FVectorVMBatchState *BatchState, VectorRegister4f a, VectorRegister4f b)                                                            { return VectorStep(VectorSubtract(b, a)); }
 VM_FORCEINLINE const uint8 *    VVM_Exec1null_random                       (VVM_NULL_FN_ARGS)                                                                                                                   { return VVM_random(CT_MultipleLoops, InsPtr, BatchState, ExecCtx, SerializeState, CmpSerializeState, NumLoops); }
 VM_FORCEINLINE const uint8 *    VVM_Exec0null_noise                        (VVM_NULL_FN_ARGS)                                                                                                                   { return InsPtr; }
 VM_FORCEINLINE VectorRegister4f VVM_Exec2f_cmplt                           (FVectorVMBatchState *BatchState, VectorRegister4f a, VectorRegister4f b)                                                            { return VectorCompareLT(a, b); }
