- Make EVVM handle safe int division in the same way as the old VM
- The EVVM didn't handle int div overflow due to 2's compliment

#rnx
#tests Divide max negative int by -1 before it would overflow, after it will not

diff --git a/Engine/Source/Runtime/VectorVM/Private/VectorVMExperimental.cpp b/Engine/Source/Runtime/VectorVM/Private/VectorVMExperimental.cpp
--- a/Engine/Source/Runtime/VectorVM/Private/VectorVMExperimental.cpp
+++ b/Engine/Source/Runtime/VectorVM/Private/VectorVMExperimental.cpp
@@ -1258,17 +1258,33 @@ VM_FORCEINLINE bool VVM_serSyncRandom(const uint8 *InsPtr, FVectorVMBatchState *
 }
 #endif //VECTORVM_SUPPORTS_SERIALIZATION
 
+static VM_FORCEINLINE int32 VVM_SafeIntDivide(int32 Numerator, int32 Denominator)
+{
+	static constexpr int32 MinIntValue = std::numeric_limits<int32>::min();
+	static constexpr int32 MaxIntValue = std::numeric_limits<int32>::max();
+
+	if (Denominator == 0)
+	{
+		return 0;
+	}
+	else if ((Denominator == -1) && (Numerator == MinIntValue))
+	{
+		return MaxIntValue;
+	}
+
+	return Numerator / Denominator;
+}
+
 static VM_FORCEINLINE VectorRegister4i VVMIntDiv(VectorRegister4i v0, VectorRegister4i v1)
 {
 	const int32 *v0_4 = reinterpret_cast<const int32*>(&v0);
 	const int32 *v1_4 = reinterpret_cast<const int32*>(&v1);
 
 	FVVM_VUI4 res;
-
-	res.i4[0] = v1_4[0] == 0 ? 0 : (v0_4[0] / v1_4[0]);
-	res.i4[1] = v1_4[1] == 0 ? 0 : (v0_4[1] / v1_4[1]);
-	res.i4[2] = v1_4[2] == 0 ? 0 : (v0_4[2] / v1_4[2]);
-	res.i4[3] = v1_4[3] == 0 ? 0 : (v0_4[3] / v1_4[3]);
+	res.i4[0] = VVM_SafeIntDivide(v0_4[0], v1_4[0]);
+	res.i4[1] = VVM_SafeIntDivide(v0_4[1], v1_4[1]);
+	res.i4[2] = VVM_SafeIntDivide(v0_4[2], v1_4[2]);
+	res.i4[3] = VVM_SafeIntDivide(v0_4[3], v1_4[3]);
 	
 	return res.v;
 }
