- Remove the legacy VM
- Results in a small compilation win & reduction in final executable size

diff --git a/Engine/Plugins/FX/Niagara/Source/Niagara/Classes/NiagaraScript.h b/Engine/Plugins/FX/Niagara/Source/Niagara/Classes/NiagaraScript.h
--- a/Engine/Plugins/FX/Niagara/Source/Niagara/Classes/NiagaraScript.h
+++ b/Engine/Plugins/FX/Niagara/Source/Niagara/Classes/NiagaraScript.h
@@ -596,21 +596,7 @@ public:
 
 	NIAGARA_API bool HasByteCode() const;
 
-#if VECTORVM_SUPPORTS_EXPERIMENTAL
-	NIAGARA_API bool SupportsExperimentalVM() const;
 	NIAGARA_API FVectorVMOptimizeContext BuildExperimentalContext() const;
-#endif //VECTORVM_SUPPORTS_EXPERIMENTAL
-
-	NIAGARA_API void PostSerialize(const FArchive& Ar);
-};
-
-template<>
-struct TStructOpsTypeTraits<FNiagaraVMExecutableData> : public TStructOpsTypeTraitsBase2<FNiagaraVMExecutableData>
-{
-	enum
-	{
-		WithPostSerialize = true,
-	};
 };
 
 struct FNiagaraGraphCachedDataBase
diff --git a/Engine/Plugins/FX/Niagara/Source/Niagara/Classes/NiagaraScriptExecutionContext.h b/Engine/Plugins/FX/Niagara/Source/Niagara/Classes/NiagaraScriptExecutionContext.h
--- a/Engine/Plugins/FX/Niagara/Source/Niagara/Classes/NiagaraScriptExecutionContext.h
+++ b/Engine/Plugins/FX/Niagara/Source/Niagara/Classes/NiagaraScriptExecutionContext.h
@@ -159,10 +159,6 @@ struct FNiagaraScriptExecutionContextBase
 	TMap<TStatIdData const*, float> ReportStats();
 #endif
 
-#if VECTORVM_SUPPORTS_EXPERIMENTAL && VECTORVM_SUPPORTS_LEGACY
-	bool bUsingExperimentalVM;
-#endif
-
 	FNDIStageTickHandler DIStageTickHandler;
 
 	FNiagaraScriptExecutionContextBase();
@@ -190,12 +186,7 @@ struct FNiagaraScriptExecutionContextBase
 	virtual bool GeneratePerInstanceDIFunctionTable(FNiagaraSystemInstance* Inst, TArray<struct FNiagaraPerInstanceDIFuncInfo>& OutFunctions) {return true;}
 
 private:
-#if VECTORVM_SUPPORTS_EXPERIMENTAL
-	bool ExecuteInternal_Experimental(uint32 NumInstances, const FScriptExecutionConstantBufferTable& ConstantBufferTable);
-#endif // VECTORVM_SUPPORTS_EXPERIMENTAL
-#if VECTORVM_SUPPORTS_LEGACY
-	bool ExecuteInternal_Legacy(uint32 NumInstances, const FScriptExecutionConstantBufferTable& ConstantBufferTable);
-#endif // VECTORVM_SUPPORTS_LEGACY
+	bool ExecuteInternal(uint32 NumInstances, const FScriptExecutionConstantBufferTable& ConstantBufferTable);
 };
 
 struct FNiagaraScriptExecutionContext : public FNiagaraScriptExecutionContextBase
@@ -239,11 +230,6 @@ protected:
 	*/
 	TArray<FNiagaraSystemInstance*>* SystemInstances;
 
-#if VECTORVM_SUPPORTS_LEGACY
-	/** Helper function that handles calling into per instance DI calls and massages the VM context appropriately. */
-	void PerInstanceFunctionHook(FVectorVMExternalFunctionContext& ParentContext, FVectorVMExternalFunctionContextLegacy& Context, int32 PerInstFunctionIndex, int32 UserPtrIndex);
-#endif
-
 public:
 	FNiagaraSystemScriptExecutionContext(ENiagaraSystemSimulationScript InScriptType) : SystemInstances(nullptr) { ScriptType = InScriptType; }
 	
diff --git a/Engine/Plugins/FX/Niagara/Source/Niagara/Classes/NiagaraScriptSourceBase.h b/Engine/Plugins/FX/Niagara/Source/Niagara/Classes/NiagaraScriptSourceBase.h
--- a/Engine/Plugins/FX/Niagara/Source/Niagara/Classes/NiagaraScriptSourceBase.h
+++ b/Engine/Plugins/FX/Niagara/Source/Niagara/Classes/NiagaraScriptSourceBase.h
@@ -72,7 +72,6 @@ public:
 	NIAGARA_API static const FString GpuScriptDefine;
 	NIAGARA_API static const FString EventSpawnDefine;
 	NIAGARA_API static const FString EventSpawnInitialAttribWritesDefine;
-	NIAGARA_API static const FString ExperimentalVMDisabled;
 	NIAGARA_API static const FString AccurateQuatInterpolation;
 
 	FNiagaraCompileOptions() : TargetUsage(ENiagaraScriptUsage::Function), TargetUsageBitmask(0)
diff --git a/Engine/Plugins/FX/Niagara/Source/Niagara/Classes/NiagaraSystem.h b/Engine/Plugins/FX/Niagara/Source/Niagara/Classes/NiagaraSystem.h
--- a/Engine/Plugins/FX/Niagara/Source/Niagara/Classes/NiagaraSystem.h
+++ b/Engine/Plugins/FX/Niagara/Source/Niagara/Classes/NiagaraSystem.h
@@ -533,8 +533,6 @@ public:
 	bool SupportsLargeWorldCoordinates() const { return bSupportLargeWorldCoordinates && bLwcEnabledSettingCached; }
 	ENiagaraLwcTileUpdateMode GetLargeWorldCoordinateTileUpdateMode() const;
 
-	FORCEINLINE bool ShouldDisableExperimentalVM() const { return bDisableExperimentalVM; }
-
 	FORCEINLINE bool IsInitialOwnerVelocityFromActor() const { return bInitialOwnerVelocityFromActor; }
 
 #if WITH_EDITORONLY_DATA
@@ -663,10 +661,6 @@ public:
 	UPROPERTY(EditAnywhere, BlueprintReadOnly, Category="Rendering", meta=(DisplayName="Default Render CustomDepth Pass", EditCondition="bOverrideRenderCustomDepth"))
 	uint8 bRenderCustomDepth : 1;
 
-	/** If true, disables optimized VM, if available */
-	UPROPERTY(EditAnywhere, AdvancedDisplay, Category = "Performance", meta = (DisplayName = "Disable Optimized VM"))
-	uint8 bDisableExperimentalVM : 1;
-
 private:
 	/*
 	When enabled we use the owner actor's velocity for the first frame.
diff --git a/Engine/Plugins/FX/Niagara/Source/Niagara/Niagara.Build.cs b/Engine/Plugins/FX/Niagara/Source/Niagara/Niagara.Build.cs
--- a/Engine/Plugins/FX/Niagara/Source/Niagara/Niagara.Build.cs
+++ b/Engine/Plugins/FX/Niagara/Source/Niagara/Niagara.Build.cs
@@ -64,12 +64,5 @@ public class Niagara : ModuleRules
 				"Slate"
             });
         }
-
-		PublicDefinitions.AddRange(
-            new string[]
-            {
-                "VECTORVM_SUPPORTS_EXPERIMENTAL=1",
-                "VECTORVM_SUPPORTS_LEGACY=1"
-            });
 	}
 }
diff --git a/Engine/Plugins/FX/Niagara/Source/Niagara/Private/IspcTestNiagaraDataInterfaceVectorField.cpp b/Engine/Plugins/FX/Niagara/Source/Niagara/Private/IspcTestNiagaraDataInterfaceVectorField.cpp
--- a/Engine/Plugins/FX/Niagara/Source/Niagara/Private/IspcTestNiagaraDataInterfaceVectorField.cpp
+++ b/Engine/Plugins/FX/Niagara/Source/Niagara/Private/IspcTestNiagaraDataInterfaceVectorField.cpp
@@ -8,12 +8,10 @@
 #include "VectorField/VectorField.h"
 #include "VectorVM.h"
 #include "VectorVMExperimental.h"
-#include "VectorVMLegacy.h"
 #include "VectorField/VectorFieldStatic.h"
 
 extern bool GNiagaraVectorFieldUseIspc;
 
-#if VECTORVM_SUPPORTS_EXPERIMENTAL && VECTORVM_SUPPORTS_LEGACY
 IMPLEMENT_SIMPLE_AUTOMATION_TEST(FIspcTestNiagaraDataInterfaceVectorField, "Ispc.Graphics.NiagaraDataInterfaceVectorField", EAutomationTestFlags::EditorContext | EAutomationTestFlags::EngineFilter)
 bool FIspcTestNiagaraDataInterfaceVectorField::RunTest(const FString& Parameters)
 {
@@ -89,8 +87,7 @@ bool FIspcTestNiagaraDataInterfaceVectorField::RunTest(const FString& Parameters
 		uint8 IncJunk[RegisterCount] = { 0, 1, 2, 3, 4, 5 };
 		uint16 VecJunk[RegisterCount] = { 0, 1, 2, 3, 4, 5 };
 		Experimental.RegInc = IncJunk;
-		FVectorVMExternalFunctionContextProxy Proxy(Experimental);
-		VectorField->SampleVectorField(Proxy);
+		VectorField->SampleVectorField(Experimental);
 	}
 
 	GNiagaraVectorFieldUseIspc = false;
@@ -153,8 +150,7 @@ bool FIspcTestNiagaraDataInterfaceVectorField::RunTest(const FString& Parameters
 		uint8 IncJunk[RegisterCount] = { 0, 1, 2, 3, 4, 5 };
 		uint16 VecJunk[RegisterCount] = { 0, 1, 2, 3, 4, 5 };
 		Experimental.RegInc = IncJunk;
-		FVectorVMExternalFunctionContextProxy Proxy(Experimental);
-		VectorField->SampleVectorField(Proxy);
+		VectorField->SampleVectorField(Experimental);
 	}
 
 	GNiagaraVectorFieldUseIspc = InitialState;
@@ -165,7 +161,7 @@ bool FIspcTestNiagaraDataInterfaceVectorField::RunTest(const FString& Parameters
 
 	return true;
 }
-#endif
+
 #if 0
 IMPLEMENT_SIMPLE_AUTOMATION_TEST(FIspcTestNiagaraDataInterfaceVectorField, "Ispc.Graphics.NiagaraDataInterfaceVectorField", EAutomationTestFlags::EditorContext | EAutomationTestFlags::EngineFilter)
 bool FIspcTestNiagaraDataInterfaceVectorField::RunTest(const FString& Parameters)
diff --git a/Engine/Plugins/FX/Niagara/Source/Niagara/Private/NiagaraAsyncCompile.cpp b/Engine/Plugins/FX/Niagara/Source/Niagara/Private/NiagaraAsyncCompile.cpp
--- a/Engine/Plugins/FX/Niagara/Source/Niagara/Private/NiagaraAsyncCompile.cpp
+++ b/Engine/Plugins/FX/Niagara/Source/Niagara/Private/NiagaraAsyncCompile.cpp
@@ -414,14 +414,6 @@ void FNiagaraAsyncCompileTask::OptimizeByteCode()
 {
 	check(ExeData);
 
-#if VECTORVM_SUPPORTS_LEGACY
-	if (ScriptPair.CompileId.AdditionalDefines.Contains(FNiagaraCompileOptions::ExperimentalVMDisabled))
-	{
-		return;
-	}
-#endif
-
-#if VECTORVM_SUPPORTS_EXPERIMENTAL
 	if (const uint8* ByteCode = ExeData->ByteCode.GetDataPtr())
 	{
 		const double OptimizeStartTime = FPlatformTime::Seconds();
@@ -452,7 +444,6 @@ void FNiagaraAsyncCompileTask::OptimizeByteCode()
 
 		CompileMetrics.ByteCodeOptimizeTime = (float) (FPlatformTime::Seconds() - OptimizeStartTime);
 	}
-#endif // VECTORVM_SUPPORTS_EXPERIMENTAL
 }
 
 void FNiagaraAsyncCompileTask::ProcessNonCompilableScript()
diff --git a/Engine/Plugins/FX/Niagara/Source/Niagara/Private/NiagaraCommon.cpp b/Engine/Plugins/FX/Niagara/Source/Niagara/Private/NiagaraCommon.cpp
--- a/Engine/Plugins/FX/Niagara/Source/Niagara/Private/NiagaraCommon.cpp
+++ b/Engine/Plugins/FX/Niagara/Source/Niagara/Private/NiagaraCommon.cpp
@@ -1693,7 +1693,6 @@ const FString FNiagaraCompileOptions::CpuScriptDefine = TEXT("CPUSim");
 const FString FNiagaraCompileOptions::GpuScriptDefine = TEXT("GPUComputeSim");
 const FString FNiagaraCompileOptions::EventSpawnDefine = TEXT("EventSpawn");
 const FString FNiagaraCompileOptions::EventSpawnInitialAttribWritesDefine = TEXT("EventSpawnInitialAttribWrites");
-const FString FNiagaraCompileOptions::ExperimentalVMDisabled = TEXT("ExperimentalVMDisabled");
 const FString FNiagaraCompileOptions::AccurateQuatInterpolation = TEXT("AccurateQuatInterpolation");
 
 FSynchronizeWithParameterDefinitionsArgs::FSynchronizeWithParameterDefinitionsArgs()
diff --git a/Engine/Plugins/FX/Niagara/Source/Niagara/Private/NiagaraDataInterfaceParticleRead.cpp b/Engine/Plugins/FX/Niagara/Source/Niagara/Private/NiagaraDataInterfaceParticleRead.cpp
--- a/Engine/Plugins/FX/Niagara/Source/Niagara/Private/NiagaraDataInterfaceParticleRead.cpp
+++ b/Engine/Plugins/FX/Niagara/Source/Niagara/Private/NiagaraDataInterfaceParticleRead.cpp
@@ -1208,26 +1208,6 @@ static const int32* GetSpawnedIDs(FVectorVMExternalFunctionContext& Context, FNi
 
 	const FNiagaraDataSet& DataSet = EmitterInstance->GetParticleData();
 
-#if VECTORVM_SUPPORTS_EXPERIMENTAL && VECTORVM_SUPPORTS_LEGACY
-	if (Context.UsingExperimentalVM)
-	{
-		const int32* SpawnedIDs = DataSet.GetFreeIDTable().GetData() + DataSet.GetFreeIDTable().Max();
-		NumSpawned = DataSet.NumSpawnedIDs;
-		if (SpawnedIDs)
-		{
-			SpawnedIDs -= NumSpawned;
-		}
-
-		return SpawnedIDs;
-	}
-	else
-	{
-		const TArray<int32>& SpawnedIDsTable = DataSet.GetSpawnedIDsTable();
-		NumSpawned = SpawnedIDsTable.Num();
-		return SpawnedIDsTable.GetData();
-	}
-#elif VECTORVM_SUPPORTS_EXPERIMENTAL
-
 	const int32* SpawnedIDs = DataSet.GetFreeIDTable().GetData() + DataSet.GetFreeIDTable().Max();
 	NumSpawned = DataSet.NumSpawnedIDs;
 	if (SpawnedIDs)
@@ -1236,14 +1216,6 @@ static const int32* GetSpawnedIDs(FVectorVMExternalFunctionContext& Context, FNi
 	}
 
 	return SpawnedIDs;
-
-#elif VECTORVM_SUPPORTS_LEGACY
-
-	const TArray<int32>& SpawnedIDsTable = DataSet.GetSpawnedIDsTable();
-	NumSpawned = SpawnedIDsTable.Num();
-	return SpawnedIDsTable.GetData();
-
-#endif
 }
 
 void UNiagaraDataInterfaceParticleRead::GetSpawnedIDAtIndex(FVectorVMExternalFunctionContext& Context)
diff --git a/Engine/Plugins/FX/Niagara/Source/Niagara/Private/NiagaraDataSet.cpp b/Engine/Plugins/FX/Niagara/Source/Niagara/Private/NiagaraDataSet.cpp
--- a/Engine/Plugins/FX/Niagara/Source/Niagara/Private/NiagaraDataSet.cpp
+++ b/Engine/Plugins/FX/Niagara/Source/Niagara/Private/NiagaraDataSet.cpp
@@ -393,11 +393,6 @@ void FNiagaraDataSet::Allocate(int32 NumInstances, bool bMaintainExisting)
 			RequiredIDs = ExistingNumIDs;
 		}
 
-#if VECTORVM_SUPPORTS_LEGACY //the new VM doesn't use the SpawnedIDsTable, so remove it.
-		// We know that we can't spawn more than NumFreeIDs particles, so we can pre-allocate SpawnedIDsTable here, to avoid allocations during execution.
-		SpawnedIDsTable.Reserve(NumFreeIDs);
-#endif
-
 		// We need to clear the ID to index table to -1 so we don't have stale entries for particles which died in the previous
 		// frame (when the results were written to another buffer). All the entries which are in use will be filled in by the script.
 		DestinationIDTable.SetNumUninitialized(RequiredIDs);
diff --git a/Engine/Plugins/FX/Niagara/Source/Niagara/Private/NiagaraEmitterInstanceImpl.cpp b/Engine/Plugins/FX/Niagara/Source/Niagara/Private/NiagaraEmitterInstanceImpl.cpp
--- a/Engine/Plugins/FX/Niagara/Source/Niagara/Private/NiagaraEmitterInstanceImpl.cpp
+++ b/Engine/Plugins/FX/Niagara/Source/Niagara/Private/NiagaraEmitterInstanceImpl.cpp
@@ -1626,9 +1626,7 @@ void FNiagaraEmitterInstanceImpl::Tick(float DeltaSeconds)
 	uint32 TotalActualEventSpawns = 0;
 
 	Data.NumSpawnedIDs = 0;
-#if VECTORVM_SUPPORTS_LEGACY
-	Data.GetSpawnedIDsTable().SetNum(0, EAllowShrinking::No);
-#endif
+
 	int32 SpawnCountRemaining = AllocationSize - OrigNumParticles;
 
 	//Init new particles with the spawn script.
diff --git a/Engine/Plugins/FX/Niagara/Source/Niagara/Private/NiagaraScript.cpp b/Engine/Plugins/FX/Niagara/Source/Niagara/Private/NiagaraScript.cpp
--- a/Engine/Plugins/FX/Niagara/Source/Niagara/Private/NiagaraScript.cpp
+++ b/Engine/Plugins/FX/Niagara/Source/Niagara/Private/NiagaraScript.cpp
@@ -437,16 +437,10 @@ FVersionedNiagaraScriptData::FVersionedNiagaraScriptData()
 }
 #endif
 
-#if VECTORVM_SUPPORTS_EXPERIMENTAL
-bool FNiagaraVMExecutableData::SupportsExperimentalVM() const
-{
-	return ExperimentalContextData.Num() > 0;
-}
-
 FVectorVMOptimizeContext FNiagaraVMExecutableData::BuildExperimentalContext() const
 {
 	FVectorVMOptimizeContext Context;
-	if (SupportsExperimentalVM())
+	if (ExperimentalContextData.Num())
 	{
 		ReinterpretVectorVMOptimizeContextData(ExperimentalContextData, Context);
 	}
@@ -457,56 +451,10 @@ FVectorVMOptimizeContext FNiagaraVMExecutableData::BuildExperimentalContext() co
 
 	return Context;
 }
-#endif
 
 bool FNiagaraVMExecutableData::HasByteCode() const
 {
-#if VECTORVM_SUPPORTS_LEGACY
-	if (ByteCode.HasByteCode() || OptimizedByteCode.HasByteCode())
-	{
-		return true;
-	}
-#endif
-
-#if VECTORVM_SUPPORTS_EXPERIMENTAL
-	if (!ExperimentalContextData.IsEmpty())
-	{
-		return true;
-	}
-#endif
-
-	return false;
-}
-
-void FNiagaraVMExecutableData::PostSerialize(const FArchive& Ar)
-{
-#if VECTORVM_SUPPORTS_EXPERIMENTAL && VECTORVM_SUPPORTS_LEGACY && !WITH_EDITOR
-	if (Ar.IsLoading())
-	{
-		// only worry about stripping out byte code if we have both sets loaded
-		if (ByteCode.GetLength() > 0 && !ExperimentalContextData.IsEmpty())
-		{
-			ENiagaraStripScriptByteCodeOption StripOption = GetDefault<UNiagaraSettings>()->ByteCodeStripOption;
-
-			const bool bStripOriginal = GNiagaraStripByteCodeOverride
-				? (GNiagaraStripByteCodeOverride > 0)
-				: StripOption == ENiagaraStripScriptByteCodeOption::Strip_Original;
-
-			const bool bStripExperimental = GNiagaraStripByteCodeOverride
-				? (GNiagaraStripByteCodeOverride < 0)
-				: StripOption == ENiagaraStripScriptByteCodeOption::Strip_Experimental;
-
-			if (bStripOriginal)
-			{
-				ByteCode.Reset();
-			}
-			else if (bStripExperimental)
-			{
-				ExperimentalContextData.Empty();
-			}
-		}
-	}
-#endif
+	return !ExperimentalContextData.IsEmpty();
 }
 
 UNiagaraScript::UNiagaraScript()
@@ -1149,11 +1097,6 @@ void UNiagaraScript::ComputeVMCompilationId(FNiagaraVMExecutableDataId& Id, cons
 				Id.AdditionalDefines.Add(TEXT("IgnoreParticleReadsForAttributeTrim"));
 			}
 
-			if (EmitterOwner->ShouldDisableExperimentalVM() || !NiagaraSettings->bExperimentalVMEnabled)
-			{
-				Id.AdditionalDefines.Add(FNiagaraCompileOptions::ExperimentalVMDisabled);
-			}
-
 			bool TrimAttributes = EmitterOwner->ShouldTrimAttributes();
 			if (TrimAttributes)
 			{
@@ -1366,11 +1309,6 @@ void UNiagaraScript::ComputeVMCompilationId(FNiagaraVMExecutableDataId& Id, cons
 			Id.AdditionalDefines.Add(TEXT("CompressAttributes"));
 		}
 
-		if (System->ShouldDisableExperimentalVM() || !NiagaraSettings->bExperimentalVMEnabled)
-		{
-			Id.AdditionalDefines.Add(FNiagaraCompileOptions::ExperimentalVMDisabled);
-		}
-
 		for (const FNiagaraEmitterHandle& EmitterHandle: System->GetEmitterHandles())
 		{
 			FVersionedNiagaraEmitterData* HandleData = EmitterHandle.GetEmitterData();
@@ -1743,20 +1681,6 @@ void FNiagaraScriptAsyncOptimizeTaskState::OptimizeByteCode()
 {		
 	SourceByteCode.Uncompress();
 
-	if (bShouldOptimizeByteCode && SourceByteCode.HasByteCode())
-	{
-#if VECTORVM_SUPPORTS_LEGACY
-		// Generate optimized byte code on any thread
-		OptimizedByteCode.Reserve(SourceByteCode.GetLength());
-		VectorVM::OptimizeByteCode(SourceByteCode.GetDataPtr(), OptimizedByteCode.GetData(), MakeArrayView(ExternalFunctionRegisterCounts));
-		OptimizedByteCode.Shrink();
-		
-		if (bShouldFreeSourceByteCodeOnCooked && OptimizedByteCode.HasByteCode())
-		{
-			SourceByteCode.Reset();
-		}
-#endif
-	}
 	bOptimizationComplete = true;
 }
 
@@ -2578,22 +2502,6 @@ void UNiagaraScript::GenerateStatIDs()
 	}
 #elif ENABLE_STATNAMEDEVENTS
 	StatNamedEvents.Empty();
-
-	if (!CachedScriptVM.SupportsExperimentalVM())
-	{
-		static const IConsoleVariable* CVarOptimizeVMDetailedStats = IConsoleManager::Get().FindConsoleVariable(TEXT("vm.DetailedVMScriptStats"));
-		if (CVarOptimizeVMDetailedStats && CVarOptimizeVMDetailedStats->GetInt() != 0)
-		{
-			if (IsReadyToRun(ENiagaraSimTarget::CPUSim))
-			{
-				StatNamedEvents.Reserve(CachedScriptVM.StatScopes.Num());
-				for (FNiagaraStatScope& StatScope : CachedScriptVM.StatScopes)
-				{
-					StatNamedEvents.Add(StatScope.FriendlyName.ToString());
-				}
-			}
-		}
-	}
 #endif
 }
 
diff --git a/Engine/Plugins/FX/Niagara/Source/Niagara/Private/NiagaraScriptExecutionContext.cpp b/Engine/Plugins/FX/Niagara/Source/Niagara/Private/NiagaraScriptExecutionContext.cpp
--- a/Engine/Plugins/FX/Niagara/Source/Niagara/Private/NiagaraScriptExecutionContext.cpp
+++ b/Engine/Plugins/FX/Niagara/Source/Niagara/Private/NiagaraScriptExecutionContext.cpp
@@ -39,18 +39,13 @@ FNiagaraScriptExecutionContextBase::FNiagaraScriptExecutionContextBase()
 	, bAllowParallel(true)
 	, bHasDIsWithPreStageTick(false)
 	, bHasDIsWithPostStageTick(false)
-#if VECTORVM_SUPPORTS_EXPERIMENTAL && VECTORVM_SUPPORTS_LEGACY
-	, bUsingExperimentalVM(true)
-#endif
 {
 
 }
 
 FNiagaraScriptExecutionContextBase::~FNiagaraScriptExecutionContextBase()
 {
-#if VECTORVM_SUPPORTS_EXPERIMENTAL
 	FreeVectorVMState(VectorVMState);
-#endif 
 }
 
 bool FNiagaraScriptExecutionContextBase::Init(FNiagaraSystemInstance* Instance, UNiagaraScript* InScript, ENiagaraSimTarget InTarget)
@@ -65,26 +60,8 @@ bool FNiagaraScriptExecutionContextBase::Init(FNiagaraSystemInstance* Instance,
 
 	HasInterpolationParameters = Script && Script->GetComputedVMCompilationId().HasInterpolatedParameters();
 
-#if VECTORVM_SUPPORTS_EXPERIMENTAL && VECTORVM_SUPPORTS_LEGACY
-	const bool ProjectSupportsExperimentalVM = GetDefault<UNiagaraSettings>()->bExperimentalVMEnabled;
-	const bool ScriptSupportsExperimentalVM = (Script ? Script->GetVMExecutableData().SupportsExperimentalVM() : false);
-	bool SystemDisabledExperimentalVM = false;
-
-	if (const UNiagaraSystem* OwnerSystem = Script->GetTypedOuter<UNiagaraSystem>())
-	{
-		if (OwnerSystem->ShouldDisableExperimentalVM())
-		{
-			SystemDisabledExperimentalVM = true;
-		}
-	}
-
-	bUsingExperimentalVM = ProjectSupportsExperimentalVM && ScriptSupportsExperimentalVM && ((GbForceExecVMPath > 0) || (!GbForceExecVMPath && !SystemDisabledExperimentalVM));
-#endif
-
-#if VECTORVM_SUPPORTS_EXPERIMENTAL
 	FVectorVMOptimizeContext OptimizeContext = Script->GetVMExecutableData().BuildExperimentalContext();
 	VectorVMState = AllocVectorVMState(&OptimizeContext);
-#endif
 
 	//If the instance is null, aka a system script, we need to pre calculate whether we'll have pre/post ticks or not.
 	if(Instance == nullptr)
@@ -228,24 +205,7 @@ bool FNiagaraScriptExecutionContextBase::Execute(class FNiagaraSystemInstance* I
 		CreateStatScopeData();
 #endif	
 
-		bool bSuccess = true;
-
-#if VECTORVM_SUPPORTS_EXPERIMENTAL && VECTORVM_SUPPORTS_LEGACY
-		if (bUsingExperimentalVM)
-		{
-			bSuccess = ExecuteInternal_Experimental(NumInstances, ConstantBufferTable);
-		}
-		else
-		{
-			bSuccess = ExecuteInternal_Legacy(NumInstances, ConstantBufferTable);
-		}
-#elif VECTORVM_SUPPORTS_EXPERIMENTAL
-		bSuccess = ExecuteInternal_Experimental(NumInstances, ConstantBufferTable);
-#elif VECTORVM_SUPPORTS_LEGACY
-		bSuccess = ExecuteInternal_Legacy(NumInstances, ConstantBufferTable);
-#else
-	#error "Not implemented"
-#endif
+		const bool bSuccess = ExecuteInternal(NumInstances, ConstantBufferTable);
 
 		// Tell the datasets we wrote how many instances were actually written.
 		for (int Idx = 0; Idx < DataSetInfo.Num(); Idx++)
@@ -273,9 +233,7 @@ bool FNiagaraScriptExecutionContextBase::Execute(class FNiagaraSystemInstance* I
 	return true;//TODO: Error cases?
 }
 
-#if VECTORVM_SUPPORTS_EXPERIMENTAL
-
-bool FNiagaraScriptExecutionContextBase::ExecuteInternal_Experimental(uint32 NumInstances, const FScriptExecutionConstantBufferTable& ConstantBufferTable)
+bool FNiagaraScriptExecutionContextBase::ExecuteInternal(uint32 NumInstances, const FScriptExecutionConstantBufferTable& ConstantBufferTable)
 {
 	TRACE_CPUPROFILER_EVENT_SCOPE(VectorVM_Experimental);
 #if STATS
@@ -306,125 +264,6 @@ bool FNiagaraScriptExecutionContextBase::ExecuteInternal_Experimental(uint32 Num
 	return true;
 }
 
-#endif
-
-#if VECTORVM_SUPPORTS_LEGACY
-
-bool FNiagaraScriptExecutionContextBase::ExecuteInternal_Legacy(uint32 NumInstances, const FScriptExecutionConstantBufferTable& ConstantBufferTable)
-{
-	TRACE_CPUPROFILER_EVENT_SCOPE(VectorVM_Legacy);
-
-	FNiagaraVMExecutableData& ExecData = Script->GetVMExecutableData();
-
-	// If we have an optimization task it must be ready at this point
-	// However we will need to lock and test again as multiple threads may be coming in here
-	if (ExecData.OptimizationTask.State.IsValid())
-	{
-		FScopeLock Lock(&ExecData.OptimizationTask.Lock);
-		if (ExecData.OptimizationTask.State.IsValid())
-		{
-			ExecData.ApplyFinishedOptimization(Script->GetVMExecutableDataCompilationId(), ExecData.OptimizationTask.State);
-			ExecData.OptimizationTask.State.Reset();
-		}
-	}
-
-	check((ExecData.ByteCode.HasByteCode() && !ExecData.ByteCode.IsCompressed()) || (ExecData.OptimizedByteCode.HasByteCode() && !ExecData.OptimizedByteCode.IsCompressed()));
-
-	VectorVM::FVectorVMExecArgs ExecArgs;
-	ExecArgs.ByteCode = ExecData.ByteCode.GetDataPtr();
-	ExecArgs.OptimizedByteCode = ExecData.OptimizedByteCode.HasByteCode() ? ExecData.OptimizedByteCode.GetDataPtr() : nullptr;
-	ExecArgs.NumTempRegisters = ExecData.NumTempRegisters;
-	ExecArgs.ConstantTableCount = ConstantBufferTable.Buffers.Num();
-	ExecArgs.ConstantTable = ConstantBufferTable.Buffers.GetData();
-	ExecArgs.ConstantTableSizes = ConstantBufferTable.BufferSizes.GetData();
-	ExecArgs.DataSetMetaTable = DataSetMetaTable;
-	ExecArgs.ExternalFunctionTable = FunctionTable.GetData();
-	ExecArgs.UserPtrTable = UserPtrTable.GetData();
-	ExecArgs.NumInstances = NumInstances;
-#if STATS
-	ExecArgs.StatScopes = MakeArrayView(StatScopeData);
-#elif ENABLE_STATNAMEDEVENTS
-	ExecArgs.StatNamedEventsScopes = Script->GetStatNamedEvents();
-#endif
-
-	ExecArgs.bAllowParallel = bAllowParallel;
-	FVectorVMSerializeState UESerializeState = { };
-#if VECTORVM_SUPPORTS_SERIALIZATION
-	static const IConsoleVariable* CVarInstancesPerChunk = IConsoleManager::Get().FindConsoleVariable(TEXT("vm.InstancesPerChunk"));
-	int32 NumParallelInstancesPerChunk = CVarInstancesPerChunk ? CVarInstancesPerChunk->GetInt() : 128;
-
-	UESerializeState.ReallocFn          = VVMRealloc;
-	UESerializeState.FreeFn             = VVMFree;
-	UESerializeState.NumInstances       = NumInstances;
-	UESerializeState.NumTempRegisters   = ExecData.NumTempRegisters;
-			
-	UESerializeState.RegisterTableFlags = (unsigned char *)VVMRealloc(NULL, UESerializeState.NumTempRegisters, __FILE__, __LINE__);
-	FMemory::Memset(UESerializeState.RegisterTableFlags, 0, UESerializeState.NumTempRegisters);
-	UESerializeState.Bytecode           = (unsigned char *)ExecData.ByteCode.GetData().GetData();
-	UESerializeState.NumBytecodeBytes   = ExecData.ByteCode.GetData().Num();
-	SerializeVectorVMInputDataSets(&UESerializeState, DataSetMetaTable,  ConstantBufferTable.Buffers.GetData(), ConstantBufferTable.BufferSizes.GetData(), ConstantBufferTable.Buffers.Num());
-
-	UESerializeState.NumChunks          = (NumInstances + NumParallelInstancesPerChunk - 1) / NumParallelInstancesPerChunk;
-	UESerializeState.Chunks             = (FVectorVMSerializeChunk *)VVMRealloc(NULL, sizeof(FVectorVMSerializeChunk) * UESerializeState.NumChunks, __FILE__, __LINE__);
-
-		
-	int NumExternalFunctions = FunctionTable.Num();
-	if (NumExternalFunctions != 0)
-	{
-		UESerializeState.ExternalData = (FVectorVMSerializeExternalData *)VVMRealloc(NULL, sizeof(FVectorVMSerializeExternalData) * NumExternalFunctions, __FILE__, __LINE__);
-		if (UESerializeState.ExternalData)
-		{
-			UESerializeState.NumExternalData = NumExternalFunctions;
-			const FNiagaraVMExecutableData& ScriptExecutableData = Script->GetVMExecutableData();
-			for (int i = 0; i < NumExternalFunctions; ++i)
-			{
-				FVectorVMSerializeExternalData *ExtData = UESerializeState.ExternalData + i;
-				TCHAR buff[128];
-				uint32 NameLen = ScriptExecutableData.CalledVMExternalFunctions[i].Name.ToString(buff, sizeof(buff) >> 1);
-				if (NameLen > 0)
-				{
-					ExtData->Name       = (wchar_t *)UESerializeState.ReallocFn(NULL, sizeof(wchar_t) * NameLen, __FILE__, __LINE__);
-					FMemory::Memcpy(ExtData->Name, buff, sizeof(TCHAR) * NameLen);
-					ExtData->NameLen    = NameLen;
-				}
-				else
-				{
-					ExtData->Name       = NULL;
-					ExtData->NameLen    = 0;
-				}
-				ExtData->NumInputs      = ScriptExecutableData.CalledVMExternalFunctions[i].GetNumInputs();
-				ExtData->NumOutputs     = ScriptExecutableData.CalledVMExternalFunctions[i].GetNumOutputs();
-			}
-		}
-	}
-#endif //VECTORVM_SUPPORTS_SERIALIZATION
-	VectorVM::Exec(ExecArgs, &UESerializeState);
-
-#if VECTORVM_SUPPORTS_SERIALIZATION
-		SerializeVectorVMOutputDataSets(&UESerializeState, DataSetMetaTable,  ConstantBufferTable.Buffers.GetData(), ConstantBufferTable.BufferSizes.GetData(), ConstantBufferTable.Buffers.Num());
-#endif //VECTORVM_SUPPORTS_SERIALIZATION
-
-	for (int32 Idx = 0; Idx < DataSetInfo.Num(); Idx++)
-	{
-		FNiagaraDataSetExecutionInfo& Info = DataSetInfo[Idx];
-
-		Info.DataSet->NumSpawnedIDs = Info.DataSet->GetSpawnedIDsTable().Num();
-	}
-
-#		if VECTORVM_SUPPORTS_SERIALIZATION
-		{
-			//we don't own the bytecode memory so we can't free it
-			UESerializeState.Bytecode = NULL;
-			UESerializeState.NumBytecodeBytes = 0;
-			FreeVectorVMSerializeState(&UESerializeState);
-		}
-#		endif //VECTORVM_SUPPORTS_SERIALIZATION
-
-	return true;
-}
-
-#endif // VECTORVM_SUPPORTS_LEGACY
-
 TArrayView<const uint8> FNiagaraScriptExecutionContextBase::GetScriptLiterals() const
 {
 #if WITH_EDITORONLY_DATA
@@ -611,7 +450,6 @@ void FNiagaraScriptExecutionContextBase::PostTick()
 }
 
 //////////////////////////////////////////////////////////////////////////
-#if VECTORVM_SUPPORTS_EXPERIMENTAL
 static void PerInsFn(FVectorVMExternalFunctionContext& ParentContext, FVectorVMExternalFunctionContextExperimental& PerInsFnContext, TArray<FNiagaraSystemInstance*>** SystemInstances, ENiagaraSystemSimulationScript ScriptType, int32 PerInstFunctionIndex, int32 UserPtrIdx)
 {
 	check(PerInsFnContext.DataSets.Num() > 0);
@@ -643,57 +481,6 @@ static void PerInsFn(FVectorVMExternalFunctionContext& ParentContext, FVectorVME
 		PerInsFnContext.UserPtrTable[UserPtrIdx] = SavedUserPtrData;
 	}
 }
-#endif
-
-#if VECTORVM_SUPPORTS_LEGACY
-void FNiagaraSystemScriptExecutionContext::PerInstanceFunctionHook(FVectorVMExternalFunctionContext& ParentContext, FVectorVMExternalFunctionContextLegacy& Context, int32 PerInstFunctionIndex, int32 UserPtrIndex)
-{
-	check(SystemInstances);
-
-	//This is a bit of a hack. We grab the base offset into the instance data from the primary dataset.
-	//TODO: Find a cleaner way to do this.
-	int32 InstanceOffset = Context.VectorVMContext->GetDataSetMeta(0).InstanceOffset;
-
-	//Cache context state.
-	int32 CachedContextStartInstance = Context.VectorVMContext->GetStartInstance();
-	int32 CachedContextNumInstances = Context.VectorVMContext->GetNumInstances();
-	uint8 const* CachedCodeLocation = Context.VectorVMContext->Code;
-
-	//Hack context so we can run the DI calls one by one.
-	Context.VectorVMContext->NumInstances = 1;
-
-	for (int32 i = 0; i < CachedContextNumInstances; ++i)
-	{
-		//Reset the code each iteration.
-		Context.VectorVMContext->Code = CachedCodeLocation;
-		//Offset buffer I/O to the correct instance's data.
-		Context.VectorVMContext->ExternalFunctionInstanceOffset = i;
-
-		int32 InstanceIndex = InstanceOffset + CachedContextStartInstance + i;
-		FNiagaraSystemInstance* Instance = (*SystemInstances)[InstanceIndex];
-		const FNiagaraPerInstanceDIFuncInfo& FuncInfo = Instance->GetPerInstanceDIFunction(ScriptType, PerInstFunctionIndex);
-
-		//TODO: We can embed the instance data inside the function lambda. No need for the user ptr table at all.
-		//Do this way for now to reduce overall complexity of the initial change. Doing this needs extensive boiler plate changes to most DI classes and a script recompile.
-		if (UserPtrIndex != INDEX_NONE)
-		{
-			Context.VectorVMContext->UserPtrTable[UserPtrIndex] = FuncInfo.InstData;
-		}
-
-		Context.VectorVMContext->StartInstance = InstanceIndex;
-
-		//TODO: In future for DIs where more perf is needed here we could split the DI func into an args gen and a execution.
-		//The this path could gen args from the bytecode once and just run the execution func per instance.
-		//I wonder if we could auto generate the args gen in a template func and just pass them into the DI for perf and reduced end user/author complexity.
-		FuncInfo.Function.Execute(ParentContext);
-	}
-
-	//Restore the context state.
-	Context.VectorVMContext->ExternalFunctionInstanceOffset = 0;
-	Context.VectorVMContext->StartInstance = CachedContextStartInstance;
-	Context.VectorVMContext->NumInstances = CachedContextNumInstances;
-}
-#endif
 
 bool FNiagaraSystemScriptExecutionContext::Init(FNiagaraSystemInstance* ParentSystemInstance, UNiagaraScript* InScript, ENiagaraSimTarget InTarget)
 {
@@ -752,42 +539,12 @@ bool FNiagaraSystemScriptExecutionContext::Tick(class FNiagaraSystemInstance* In
 						//Currently we must assume that any User DI is overridden but maybe we can be less conservative with this in future.
 						if (ScriptDIInfo.NeedsPerInstanceBinding())
 						{
-#if VECTORVM_SUPPORTS_EXPERIMENTAL && VECTORVM_SUPPORTS_LEGACY
-							if (bUsingExperimentalVM)
-							{
-								FuncInfo.Function = FVMExternalFunction::CreateLambda(
-									[SystemInstances = &this->SystemInstances, ScriptType = this->ScriptType, NumPerInstanceFunctions, UserPtrIdx = ScriptDIInfo.UserPtrIdx](FVectorVMExternalFunctionContext& ExtFnContext)
-									{
-										PerInsFn(ExtFnContext, ExtFnContext.Experimental, SystemInstances, ScriptType, NumPerInstanceFunctions, UserPtrIdx);
-									}
-								);
-							}
-							else
-							{
-								FuncInfo.Function = FVMExternalFunction::CreateLambda(
-									[ExecContext = this, NumPerInstanceFunctions, UserPtrIndex = ScriptDIInfo.UserPtrIdx](FVectorVMExternalFunctionContext& Context)
-									{
-										//This DI needs a binding per instance so we just bind to the external function hook which will call the correct binding for each instance.
-										ExecContext->PerInstanceFunctionHook(Context, Context.Legacy, NumPerInstanceFunctions, UserPtrIndex);
-									}
-								);
-							}
-#elif VECTORVM_SUPPORTS_EXPERIMENTAL
 							FuncInfo.Function = FVMExternalFunction::CreateLambda(
 								[SystemInstances = &this->SystemInstances, ScriptType = this->ScriptType, NumPerInstanceFunctions, UserPtrIdx = ScriptDIInfo.UserPtrIdx](FVectorVMExternalFunctionContext& ExtFnContext)
-							{
-								PerInsFn(ExtFnContext, ExtFnContext, SystemInstances, ScriptType, NumPerInstanceFunctions, UserPtrIdx);
-							}
-							);
-#elif VECTORVM_SUPPORTS_LEGACY
-							FuncInfo.Function = FVMExternalFunction::CreateLambda(
-								[ExecContext = this, NumPerInstanceFunctions, UserPtrIndex = ScriptDIInfo.UserPtrIdx](FVectorVMExternalFunctionContext& Context)
-							{
-								//This DI needs a binding per instance so we just bind to the external function hook which will call the correct binding for each instance.
-								ExecContext->PerInstanceFunctionHook(Context, Context, NumPerInstanceFunctions, UserPtrIndex);
-							}
+								{
+									PerInsFn(ExtFnContext, ExtFnContext, SystemInstances, ScriptType, NumPerInstanceFunctions, UserPtrIdx);
+								}
 							);
-#endif
 							++NumPerInstanceFunctions;
 						}
 						else
diff --git a/Engine/Plugins/FX/Niagara/Source/Niagara/Private/NiagaraSystem.cpp b/Engine/Plugins/FX/Niagara/Source/Niagara/Private/NiagaraSystem.cpp
--- a/Engine/Plugins/FX/Niagara/Source/Niagara/Private/NiagaraSystem.cpp
+++ b/Engine/Plugins/FX/Niagara/Source/Niagara/Private/NiagaraSystem.cpp
@@ -273,7 +273,6 @@ UNiagaraSystem::UNiagaraSystem(const FObjectInitializer& ObjectInitializer)
 , bCompileForEdit(false)
 #endif
 , bSupportLargeWorldCoordinates(true)
-, bDisableExperimentalVM(false)
 , CustomDepthStencilWriteMask(ERendererStencilMask::ERSM_Default)
 , bFixedBounds(false)
 #if WITH_EDITORONLY_DATA
@@ -1046,11 +1045,6 @@ void UNiagaraSystem::RecomputeExecutionOrderForDataInterface(class UNiagaraDataI
 
 ENiagaraPerformanceStateScopeMode UNiagaraSystem::GetPerformanceStateScopeMode() const
 {
-	const UNiagaraSettings* Settings = GetDefault<UNiagaraSettings>();
-	if (bDisableExperimentalVM || !Settings || !Settings->bExperimentalVMEnabled)
-	{
-		return ENiagaraPerformanceStateScopeMode::Detailed;
-	}
 	return ENiagaraPerformanceStateScopeMode::Basic;
 }
 
diff --git a/Engine/Plugins/FX/Niagara/Source/Niagara/Public/NiagaraSettings.h b/Engine/Plugins/FX/Niagara/Source/Niagara/Public/NiagaraSettings.h
--- a/Engine/Plugins/FX/Niagara/Source/Niagara/Public/NiagaraSettings.h
+++ b/Engine/Plugins/FX/Niagara/Source/Niagara/Public/NiagaraSettings.h
@@ -164,14 +164,6 @@ class UNiagaraSettings : public UDeveloperSettings
 	UPROPERTY(config, EditAnywhere, Category = Niagara, meta = ( DisplayName = "Enforce strict type checks in the graph" ))
 	bool bEnforceStrictStackTypes = true;
 
-	/**
-	 True indicates that we will generate byte code for the new optimized VM.  Control over whether the new VM will
-	 be used when executing NiagaraScripts will also take into account the overrides on the system (bDisableExperimentalVM) and
-	 the cvars fx.NiagaraScript.StripByteCodeOnLoad and fx.ForceExecVMPath.
-	*/
-	UPROPERTY(config, EditAnywhere, Category = Niagara, meta = (DisplayName = "Enable building data for Optimized VM"))
-	bool bExperimentalVMEnabled = true;
-
 	/** If set to true, quaternion attributes will be interpolated via slerp instead of lerp in interpolated spawn scripts. */
 	UPROPERTY(config, EditAnywhere, AdvancedDisplay, Category = Niagara)
 	bool bAccurateQuatInterpolation = true;
diff --git a/Engine/Plugins/FX/Niagara/Source/NiagaraEditor/Private/NiagaraCompilationTasks.cpp b/Engine/Plugins/FX/Niagara/Source/NiagaraEditor/Private/NiagaraCompilationTasks.cpp
--- a/Engine/Plugins/FX/Niagara/Source/NiagaraEditor/Private/NiagaraCompilationTasks.cpp
+++ b/Engine/Plugins/FX/Niagara/Source/NiagaraEditor/Private/NiagaraCompilationTasks.cpp
@@ -1178,16 +1178,6 @@ void FNiagaraSystemCompilationTask::FCompileTaskInfo::GenerateOptimizedVMByteCod
 {
 	TRACE_CPUPROFILER_EVENT_SCOPE(NiagaraAsyncTask_OptimizeVM);
 
-#if VECTORVM_SUPPORTS_LEGACY
-	const bool bExperimentalVMDisabled = CompileOptions.AdditionalDefines.Contains(FNiagaraCompileOptions::ExperimentalVMDisabled);
-
-	if (bExperimentalVMDisabled)
-	{
-		return;
-	}
-#endif
-
-#if VECTORVM_SUPPORTS_EXPERIMENTAL
 	if (const uint8* ByteCode = ExeData->ByteCode.GetDataPtr())
 	{
 		const double OptimizeStartTime = FPlatformTime::Seconds();
@@ -1217,7 +1207,6 @@ void FNiagaraSystemCompilationTask::FCompileTaskInfo::GenerateOptimizedVMByteCod
 
 		ByteCodeOptimizeTime = (float)(OptimizeStartTime - FPlatformTime::Seconds());
 	}
-#endif // VECTORVM_SUPPORTS_EXPERIMENTAL
 }
 
 void FNiagaraSystemCompilationTask::FCompileTaskInfo::ProcessCompilation(FNiagaraSystemCompilationTask* SystemCompileTask, const FCompileGroupInfo& GroupInfo)
diff --git a/Engine/Plugins/FX/Niagara/Source/NiagaraEditor/Private/Widgets/SNiagaraGeneratedCodeView.cpp b/Engine/Plugins/FX/Niagara/Source/NiagaraEditor/Private/Widgets/SNiagaraGeneratedCodeView.cpp
--- a/Engine/Plugins/FX/Niagara/Source/NiagaraEditor/Private/Widgets/SNiagaraGeneratedCodeView.cpp
+++ b/Engine/Plugins/FX/Niagara/Source/NiagaraEditor/Private/Widgets/SNiagaraGeneratedCodeView.cpp
@@ -384,8 +384,6 @@ void SNiagaraGeneratedCodeView::UpdateUI()
 		}
 	}
 
-	const bool bIncludeEVVM = GetDefault<UNiagaraSettings>()->bExperimentalVMEnabled;
-
 	// Mark the scripts with the correct display type and copy references for the non-gpu scripts for the assembly view.
 	int32 OriginalScriptCount = Scripts.Num();
 	ScriptDisplayTypes.AddUninitialized(OriginalScriptCount);
@@ -403,11 +401,8 @@ void SNiagaraGeneratedCodeView::UpdateUI()
 			Scripts.Add(Script);
 			ScriptDisplayTypes.Add(ASSEMBLY_SCRIPT_TYPE);
 
-			if (bIncludeEVVM)
-			{
-				Scripts.Add(Script);
-				ScriptDisplayTypes.Add(EVVM_ASSEMBLY_SCRIPT_TYPE);
-			}
+			Scripts.Add(Script);
+			ScriptDisplayTypes.Add(EVVM_ASSEMBLY_SCRIPT_TYPE);
 		}
 	}
 		
diff --git a/Engine/Source/Runtime/Core/Private/UObject/FortniteMainBranchObjectVersion.cpp b/Engine/Source/Runtime/Core/Private/UObject/FortniteMainBranchObjectVersion.cpp
--- a/Engine/Source/Runtime/Core/Private/UObject/FortniteMainBranchObjectVersion.cpp
+++ b/Engine/Source/Runtime/Core/Private/UObject/FortniteMainBranchObjectVersion.cpp
@@ -11,7 +11,7 @@ TMap<FGuid, FGuid> FFortniteMainBranchObjectVersion::GetSystemGuids()
 	SystemGuids.Add(DevGuids.MATERIALSHADERMAP_DERIVEDDATA_VER, FGuid("04B8C817071A453C9EBEFBFB51D5F785"));
 	SystemGuids.Add(DevGuids.NANITE_DERIVEDDATA_VER, FGuid("42C46B24E1586986880AFE62B0F575AB"));
 	SystemGuids.Add(DevGuids.NIAGARASHADERMAP_DERIVEDDATA_VER, FGuid("6360A977062842A29ED30E3A7ACB0E64"));
-	SystemGuids.Add(DevGuids.Niagara_LatestScriptCompileVersion, FGuid("F83165088DBF08469F4E45D0D1E806E3"));
+	SystemGuids.Add(DevGuids.Niagara_LatestScriptCompileVersion, FGuid("18843C17BE24154D9DD3EA597ADC5FF5"));
 	SystemGuids.Add(DevGuids.SkeletalMeshDerivedDataVersion, FGuid("716424BCEC314E969ACFE36D2C2ECE49"));
 	SystemGuids.Add(DevGuids.STATICMESH_DERIVEDDATA_VER, FGuid("3ABF32149F9D4D838750368AB95573BE"));
 	SystemGuids.Add(DevGuids.MaterialTranslationDDCVersion, FGuid("84CF0B79D69D4A4BBE8CAFD4B3D48AB7"));
diff --git a/Engine/Source/Runtime/VectorVM/Private/VectorVM.cpp b/Engine/Source/Runtime/VectorVM/Private/VectorVM.cpp
--- a/Engine/Source/Runtime/VectorVM/Private/VectorVM.cpp
+++ b/Engine/Source/Runtime/VectorVM/Private/VectorVM.cpp
@@ -143,34 +143,6 @@ FORCEINLINE VectorRegister4Int VectorIntShuffle(const VectorRegister4Int& Vec, c
 }
 #endif
 
-#if VECTORVM_SUPPORTS_LEGACY
-
-//Temporarily locking the free table until we can implement a lock free algorithm. UE-65856
-FORCEINLINE void FDataSetMeta::LockFreeTable()
-{
-#if FREE_TABLE_LOCK_CONTENTION_WARNINGS
-	uint64 StartCycles = FPlatformTime::Cycles64();
-#endif
- 		
-	FreeTableLock.Lock();
- 
-#if FREE_TABLE_LOCK_CONTENTION_WARNINGS
-	uint64 EndCylces = FPlatformTime::Cycles64();
-	double DurationMs = FPlatformTime::ToMilliseconds64(EndCylces - StartCycles);
-	if (DurationMs >= FREE_TABLE_LOCK_CONTENTION_WARN_THRESHOLD_MS)
-	{
-		UE_LOG(LogVectorVM, Warning, TEXT("VectorVM Stalled in LockFreeTable()! %g ms"), DurationMs);
-	}
-#endif
-}
-
-FORCEINLINE void FDataSetMeta::UnlockFreeTable()
-{
- 	FreeTableLock.Unlock();
-}
-
-#endif //VECTORVM_SUPPORTS_LEGACY
-
 static int32 GbParallelVVM = 1;
 static FAutoConsoleVariableRef CVarbParallelVVM(
 	TEXT("vm.Parallel"),
@@ -326,13 +298,6 @@ uint8 VectorVM::CreateSrcOperandMask(EVectorVMOperandLocation Type0, EVectorVMOp
 		(Type2 == EVectorVMOperandLocation::Constant ? OP2_CONST : OP_REGISTER);
 }
 
-#if VECTORVM_SUPPORTS_LEGACY
-namespace VectorKernelNoiseImpl
-{
-	static void BuildNoiseTable();
-}
-#endif
-
 void VectorVM::Init()
 {
 	static bool Inited = false;
@@ -343,3499 +308,8 @@ void VectorVM::Init()
 		g_VectorVMEnumOperandObj = StaticEnum<EVectorVMOperandLocation>();
 #endif
 
-#if VECTORVM_SUPPORTS_LEGACY
-		VectorKernelNoiseImpl::BuildNoiseTable();
-#endif
 		Inited = true;
 	}
 }
 
-#if VECTORVM_SUPPORTS_LEGACY
-
-//////////////////////////////////////////////////////////////////////////
-//  VM Code Optimizer Context
-
-typedef void(*FVectorVMExecFunction)(FVectorVMContext&);
-
-struct FVectorVMCodeOptimizerContext
-{
-	typedef bool(*OptimizeVMFunction)(EVectorVMOp, FVectorVMCodeOptimizerContext&);
-
-	explicit FVectorVMCodeOptimizerContext(FVectorVMContext& InBaseContext, const uint8* ByteCode, TArray<uint8>& InOptimizedCode, TArrayView<uint8> InExternalFunctionRegisterCounts)
-		: BaseContext(InBaseContext)
-		, OptimizedCode(InOptimizedCode)
-		, ExternalFunctionRegisterCounts(InExternalFunctionRegisterCounts)
-	{
-		BaseContext.PrepareForExec(0, 0, nullptr, nullptr, nullptr, nullptr, TArrayView<FDataSetMeta>(), 0, false);
-		BaseContext.PrepareForChunk(ByteCode, 0, 0);
-
-		// write out a jump table offset that we'll fill in when the table is encoded (EncodeJumpTable())
-		Write<uint32>(0);
-	}
-	FVectorVMCodeOptimizerContext(const FVectorVMCodeOptimizerContext&) = delete;
-	FVectorVMCodeOptimizerContext(const FVectorVMCodeOptimizerContext&&) = delete;
-
-	template<uint32 InstancesPerOp>
-	int32 GetNumLoops() const { return 0; }
-
-	FORCEINLINE EVectorVMOp PeekOp()
-	{
-		return static_cast<EVectorVMOp>(*BaseContext.Code);
-	}
-
-	FORCEINLINE uint8 DecodeU8() { return BaseContext.DecodeU8(); }
-	FORCEINLINE uint16 DecodeU16() { return BaseContext.DecodeU16(); }
-	FORCEINLINE uint32 DecodeU32() { return BaseContext.DecodeU32(); }
-	FORCEINLINE uint64 DecodeU64() { return BaseContext.DecodeU64(); }
-
-	//-TODO: Support unaligned writes
-	template<typename T>
-	void Write(const T& v)
-	{
-		reinterpret_cast<T&>(OptimizedCode[OptimizedCode.AddUninitialized(sizeof(T))]) = v;
-	}
-
-	void WriteExecFunction(const FVectorVMExecFunction& Function)
-	{
-		const int32 JumpTableIndex = JumpTable.AddUnique(Function);
-		check(JumpTableIndex <= TNumericLimits<uint8>::Max());
-		Write<uint8>(JumpTableIndex);
-	}
-
-	struct FOptimizerCodeState
-	{
-		uint8 const* BaseContextCode;
-		int32 OptimizedCodeLength;
-	};
-
-	FOptimizerCodeState CreateCodeState()
-	{
-		FOptimizerCodeState State;
-		State.BaseContextCode = BaseContext.Code;
-		State.OptimizedCodeLength = OptimizedCode.Num();
-
-		return State;
-	}
-
-	void RollbackCodeState(const FOptimizerCodeState& State)
-	{
-		BaseContext.Code = State.BaseContextCode;
-		OptimizedCode.SetNum(State.OptimizedCodeLength, EAllowShrinking::No);
-	}
-
-	// Jump table is encoded at the end of the optimized code, with the first int32 in the byte code
-	// stream being the start offset
-
-	static const FVectorVMExecFunction* DecodeJumpTable(const uint8*& OptimizedByteCode)
-	{
-		const uint32 JumpTableOffset = *reinterpret_cast<const uint32*>(OptimizedByteCode);
-		const FVectorVMExecFunction* JumpTable = reinterpret_cast<const FVectorVMExecFunction*>(OptimizedByteCode + JumpTableOffset);
-
-		OptimizedByteCode += sizeof(uint32);
-
-		return JumpTable;
-	}
-
-	void EncodeJumpTable()
-	{
-		const int32 JumpTableCount = JumpTable.Num();
-		check(JumpTableCount > 0);
-		check(JumpTableCount <= TNumericLimits<uint8>::Max());
-
-		if (JumpTableCount <= 0 || JumpTableCount > TNumericLimits<uint8>::Max())
-		{
-			// if the jump table is too big, then clear out the OptimizedCode and we'll just have to run the unoptimized path
-			OptimizedCode.Reset();
-			return;
-		}
-
-		// write the offset to the jump table into the reserved slot
-		const uint32 JumpTableOffset = OptimizedCode.Num();
-		*reinterpret_cast<uint32*>(OptimizedCode.GetData()) = JumpTableOffset;
-
-		const int32 JumpTableSize = JumpTableCount * sizeof(FVectorVMExecFunction);
-		OptimizedCode.AddUninitialized(JumpTableSize);
-		FMemory::Memcpy(OptimizedCode.GetData() + JumpTableOffset, JumpTable.GetData(), JumpTableSize);
-	}
-
-	FVectorVMContext&		BaseContext;
-	TArray<uint8>&			OptimizedCode;
-	TArray<FVectorVMExecFunction, TInlineAllocator<256>> JumpTable;
-	const TArrayView<uint8>	ExternalFunctionRegisterCounts;
-	const int32				StartInstance = 0;
-};
-
-//////////////////////////////////////////////////////////////////////////
-//  Constant Handlers
-
-struct FConstantHandlerBase
-{
-	const uint16 ConstantIndex;
-	FConstantHandlerBase(FVectorVMContext& Context)
-		: ConstantIndex(Context.DecodeU16())
-	{}
-
-	FORCEINLINE void Advance() { }
-
-	static void Optimize(FVectorVMCodeOptimizerContext& Context)
-	{
-		Context.Write(Context.DecodeU16());
-	}
-
-	static void OptimizeSkip(FVectorVMCodeOptimizerContext& Context)
-	{
-		Context.DecodeU16();
-	}
-};
-
-template<typename T>
-struct FConstantHandler : public FConstantHandlerBase
-{
-	const T Constant;
-	FConstantHandler(FVectorVMContext& Context)
-		: FConstantHandlerBase(Context)
-		, Constant(*Context.GetConstant<T>(ConstantIndex))
-	{
-	}
-
-	FORCEINLINE const T& Get() const { return Constant; }
-	FORCEINLINE const T& GetAndAdvance() { return Constant; }
-};
-
-template<>
-struct FConstantHandler<VectorRegister4Float> : public FConstantHandlerBase
-{
-	static VectorRegister4Float LoadConstant(const FVectorVMContext& Context, uint16 ConstantIndex)
-	{
-		float ConstantValue = *Context.GetConstant<float>(ConstantIndex);
-
-		return MakeVectorRegister(ConstantValue, ConstantValue, ConstantValue, ConstantValue);
-	}
-
-	const VectorRegister4Float Constant;
-	FConstantHandler(FVectorVMContext& Context)
-		: FConstantHandlerBase(Context)
-		, Constant(LoadConstant(Context, ConstantIndex))
-	{}
-
-	FORCEINLINE const VectorRegister4Float Get() const { return Constant; }
-	FORCEINLINE const VectorRegister4Float GetAndAdvance() { return Constant; }
-};
-
-template<>
-struct FConstantHandler<VectorRegister4Int> : public FConstantHandlerBase
-{
-	static VectorRegister4Int LoadConstant(const FVectorVMContext& Context, uint16 ConstantIndex)
-	{
-		int32 ConstantValue = *Context.GetConstant<int32>(ConstantIndex);
-
-		return MakeVectorRegisterInt(ConstantValue, ConstantValue, ConstantValue, ConstantValue);
-	}
-
-	const VectorRegister4Int Constant;
-	FConstantHandler(FVectorVMContext& Context)
-		: FConstantHandlerBase(Context)
-		, Constant(LoadConstant(Context, ConstantIndex))
-	{}
-
-	FORCEINLINE const VectorRegister4Int Get() const { return Constant; }
-	FORCEINLINE const VectorRegister4Int GetAndAdvance() { return Constant; }
-};
-
-
-//////////////////////////////////////////////////////////////////////////
-// Register handlers.
-// Handle reading of a register, advancing the pointer with each read.
-
-struct FRegisterHandlerBase
-{
-	const int32 RegisterIndex;
-	FORCEINLINE FRegisterHandlerBase(FVectorVMContext& Context)
-		: RegisterIndex(Context.DecodeU16())
-	{}
-
-	static void Optimize(FVectorVMCodeOptimizerContext& Context)
-	{
-		Context.Write(Context.DecodeU16());
-	}
-};
-
-template<typename T>
-struct FRegisterHandler : public FRegisterHandlerBase
-{
-private:
-	T * RESTRICT Register;
-public:
-	FORCEINLINE FRegisterHandler(FVectorVMContext& Context)
-		: FRegisterHandlerBase(Context)
-		, Register((T*)Context.GetTempRegister(RegisterIndex))
-	{}
-
-	FORCEINLINE const T Get() { return *Register; }
-	FORCEINLINE T* GetDest() { return Register; }
-	FORCEINLINE void Advance() { ++Register; }
-	FORCEINLINE const T GetAndAdvance()
-	{
-		return *Register++;
-	}
-	FORCEINLINE T* GetDestAndAdvance()
-	{
-		return Register++;
-	}
-};
-
-//////////////////////////////////////////////////////////////////////////
-
-FVectorVMContext::FVectorVMContext()
-	: Code(nullptr)
-	, ConstantTable(nullptr)
-	, ExternalFunctionTable(nullptr)
-	, UserPtrTable(nullptr)
-	, NumInstances(0)
-	, StartInstance(0)
-	, TempRegisterSize(0)
-	, TempBufferSize(0)
-{
-	RandStream.GenerateNewSeed();
-}
-
-void FVectorVMContext::PrepareForExec(
-	int32 InNumTempRegisters,
-	int32 InConstantTableCount,
-	const uint8* const* InConstantTable,
-	const int32* InConstantTableSizes,
-	const FVMExternalFunction* const* InExternalFunctionTable,
-	void** InUserPtrTable,
-	TArrayView<FDataSetMeta> InDataSetMetaTable,
-	int32 MaxNumInstances,
-	bool bInParallelExecution
-)
-{
-	NumTempRegisters = InNumTempRegisters;
-	ConstantTableCount = InConstantTableCount;
-	ConstantTableSizes = InConstantTableSizes;
-	ConstantTable = InConstantTable;
-	ExternalFunctionTable = InExternalFunctionTable;
-	UserPtrTable = InUserPtrTable;
-
-	TempRegisterSize = Align(MaxNumInstances * VectorVM::MaxInstanceSizeBytes, PLATFORM_CACHE_LINE_SIZE);
-	TempBufferSize = TempRegisterSize * NumTempRegisters;
-	TempRegTable.SetNumUninitialized(TempBufferSize, EAllowShrinking::No);
-
-	DataSetMetaTable = InDataSetMetaTable;
-
-	for (auto& TLSTempData : ThreadLocalTempData)
-	{
-		TLSTempData.Reset();
-	}
-	ThreadLocalTempData.SetNum(DataSetMetaTable.Num());
-
-	bIsParallelExecution = bInParallelExecution;
-}
-
-#if STATS
-void FVectorVMContext::SetStatScopes(TArrayView<FStatScopeData> InStatScopes)
-{
-	StatScopes = InStatScopes;
-	StatCounterStack.Reserve(StatScopes.Num());
-	ScopeExecCycles.Reset(StatScopes.Num());
-	ScopeExecCycles.AddZeroed(StatScopes.Num());
-}
-#elif ENABLE_STATNAMEDEVENTS
-void FVectorVMContext::SetStatNamedEventScopes(TArrayView<const FString> InStatNamedEventScopes)
-{
-	StatNamedEventScopes = InStatNamedEventScopes;
-}
-#endif
-
-void FVectorVMContext::FinishExec()
-{
-	//At the end of executing each chunk we can push any thread local temporary data out to the main storage with locks or atomics.
-
-	check(ThreadLocalTempData.Num() == DataSetMetaTable.Num());
-	for(int32 DataSetIndex=0; DataSetIndex < DataSetMetaTable.Num(); ++DataSetIndex)
-	{
-		FDataSetThreadLocalTempData&RESTRICT Data = ThreadLocalTempData[DataSetIndex];
-
-		if (Data.IDsToFree.Num() > 0)
-		{
-			TArray<int32>&RESTRICT FreeIDTable = *DataSetMetaTable[DataSetIndex].FreeIDTable;
-			int32&RESTRICT NumFreeIDs = *DataSetMetaTable[DataSetIndex].NumFreeIDs;
-			check(FreeIDTable.Num() >= NumFreeIDs + Data.IDsToFree.Num());
-
-			//Temporarily locking the free table until we can implement something lock-free
-			DataSetMetaTable[DataSetIndex].LockFreeTable();
-			for (int32 IDToFree : Data.IDsToFree)
-			{
-				//UE_LOG(LogVectorVM, Warning, TEXT("AddFreeID: ID:%d | FreeTableIdx:%d."), IDToFree, NumFreeIDs);
-				FreeIDTable[NumFreeIDs++] = IDToFree;
-			}
-			//Unlock the free table.
-			DataSetMetaTable[DataSetIndex].UnlockFreeTable();
-			Data.IDsToFree.Reset();
-		}
-
-		//Also update the max ID seen. This should be the ONLY place in the VM we update this max value.
-		if ( bIsParallelExecution )
-		{
-			volatile int32* MaxUsedID = DataSetMetaTable[DataSetIndex].MaxUsedID;
-			int32 LocalMaxUsedID;
-			do
-			{
-				LocalMaxUsedID = *MaxUsedID;
-				if (LocalMaxUsedID >= Data.MaxID)
-				{
-					break;
-				}
-			} while (FPlatformAtomics::InterlockedCompareExchange(MaxUsedID, Data.MaxID, LocalMaxUsedID) != LocalMaxUsedID);
-
-			*MaxUsedID = FMath::Max(*MaxUsedID, Data.MaxID);
-		}
-		else
-		{
-			int32* MaxUsedID = DataSetMetaTable[DataSetIndex].MaxUsedID;
-			*MaxUsedID = FMath::Max(*MaxUsedID, Data.MaxID);
-		}
-	}
-
-#if STATS
-	check(ScopeExecCycles.Num() == StatScopes.Num());
-	for (int i = 0; i < StatScopes.Num(); i++)
-	{
-		uint64 ExecTime = ScopeExecCycles[i];
-		if (ExecTime > 0)
-		{
-			std::atomic_fetch_add(&StatScopes[i].ExecutionCycleCount, ExecTime);
-		}
-	}
-	StatScopes = TArrayView<FStatScopeData>();
-#elif ENABLE_STATNAMEDEVENTS
-	StatNamedEventScopes = TArrayView<FString>();
-#endif
-}
-
-//////////////////////////////////////////////////////////////////////////
-
-//////////////////////////////////////////////////////////////////////////
-// Kernels
-template<typename Kernel, typename DstHandler, typename Arg0Handler, uint32 NumInstancesPerOp>
-struct TUnaryKernelHandler
-{
-	static void Optimize(FVectorVMCodeOptimizerContext& Context)
-	{
-		Context.WriteExecFunction(Exec);
-		Arg0Handler::Optimize(Context);
-		DstHandler::Optimize(Context);
-	}
-
-	static VM_FORCEINLINE void Exec(FVectorVMContext& Context)
-	{
-		Arg0Handler Arg0(Context);
-		DstHandler Dst(Context);
-
-		const int32 Loops = Context.GetNumLoops<NumInstancesPerOp>();
-		for (int32 i = 0; i < Loops; ++i)
-		{
-			Kernel::DoKernel(Context, Dst.GetDestAndAdvance(), Arg0.GetAndAdvance());
-		}
-	}
-};
-
-template<typename Kernel, typename DstHandler, typename Arg0Handler, typename Arg1Handler, int32 NumInstancesPerOp>
-struct TBinaryKernelHandler
-{
-	static void Optimize(FVectorVMCodeOptimizerContext& Context)
-	{
-		Context.WriteExecFunction(Exec);
-		Arg0Handler::Optimize(Context);
-		Arg1Handler::Optimize(Context);
-		DstHandler::Optimize(Context);
-	}
-
-	static void Exec(FVectorVMContext& Context)
-	{
-		Arg0Handler Arg0(Context); 
-		Arg1Handler Arg1(Context);
-
-		DstHandler Dst(Context);
-
-		const int32 Loops = Context.GetNumLoops<NumInstancesPerOp>();
-		for (int32 i = 0; i < Loops; ++i)
-		{
-			Kernel::DoKernel(Context, Dst.GetDestAndAdvance(), Arg0.GetAndAdvance(), Arg1.GetAndAdvance());
-		}
-	}
-};
-
-template<typename Kernel, typename DstHandler, typename Arg0Handler, typename Arg1Handler, typename Arg2Handler, int32 NumInstancesPerOp>
-struct TTrinaryKernelHandler
-{
-	static void Optimize(FVectorVMCodeOptimizerContext& Context)
-	{
-		Context.WriteExecFunction(Exec);
-		Arg0Handler::Optimize(Context);
-		Arg1Handler::Optimize(Context);
-		Arg2Handler::Optimize(Context);
-		DstHandler::Optimize(Context);
-	}
-
-	static void Exec(FVectorVMContext& Context)
-	{
-		Arg0Handler Arg0(Context);
-		Arg1Handler Arg1(Context);
-		Arg2Handler Arg2(Context);
-
-		DstHandler Dst(Context);
-
-		const int32 Loops = Context.GetNumLoops<NumInstancesPerOp>();
-		for (int32 i = 0; i < Loops; ++i)
-		{
-			Kernel::DoKernel(Context, Dst.GetDestAndAdvance(), Arg0.GetAndAdvance(), Arg1.GetAndAdvance(), Arg2.GetAndAdvance());
-		}
-	}
-};
-
-
-/** Base class of vector kernels with a single operand. */
-template <typename Kernel, typename DstHandler, typename ConstHandler, typename RegisterHandler, int32 NumInstancesPerOp>
-struct TUnaryKernel
-{
-	static void Optimize(FVectorVMCodeOptimizerContext& Context)
-	{
-		const uint32 SrcOpTypes = Context.BaseContext.DecodeSrcOperandTypes();
-		switch (SrcOpTypes)
-		{
-		case SRCOP_RRR: TUnaryKernelHandler<Kernel, DstHandler, RegisterHandler, NumInstancesPerOp>::Optimize(Context); break;
-		case SRCOP_RRC:	TUnaryKernelHandler<Kernel, DstHandler, ConstHandler, NumInstancesPerOp>::Optimize(Context); break;
-		default: check(0); break;
-		};
-	}
-
-	static void Exec(FVectorVMContext& Context)
-	{
-		const uint32 SrcOpTypes = Context.DecodeSrcOperandTypes();
-		switch (SrcOpTypes)
-		{
-		case SRCOP_RRR: TUnaryKernelHandler<Kernel, DstHandler, RegisterHandler, NumInstancesPerOp>::Exec(Context); break;
-		case SRCOP_RRC:	TUnaryKernelHandler<Kernel, DstHandler, ConstHandler, NumInstancesPerOp>::Exec(Context); break;
-		default: check(0); break; 
-		};
-	}
-};
-template<typename Kernel>
-struct TUnaryScalarKernel : public TUnaryKernel<Kernel, FRegisterHandler<float>, FConstantHandler<float>, FRegisterHandler<float>, 1> {};
-template<typename Kernel>
-struct TUnaryVectorKernel : public TUnaryKernel<Kernel, FRegisterHandler<VectorRegister4Float>, FConstantHandler<VectorRegister4Float>, FRegisterHandler<VectorRegister4Float>, VECTOR_WIDTH_FLOATS> {};
-template<typename Kernel>
-struct TUnaryScalarIntKernel : public TUnaryKernel<Kernel, FRegisterHandler<int32>, FConstantHandler<int32>, FRegisterHandler<int32>, 1> {};
-template<typename Kernel>
-struct TUnaryVectorIntKernel : public TUnaryKernel<Kernel, FRegisterHandler<VectorRegister4Int>, FConstantHandler<VectorRegister4Int>, FRegisterHandler<VectorRegister4Int>, VECTOR_WIDTH_FLOATS> {};
-
-/** Base class of Vector kernels with 2 operands. */
-template <typename Kernel, typename DstHandler, typename ConstHandler, typename RegisterHandler, uint32 NumInstancesPerOp>
-struct TBinaryKernel
-{
-	static void Optimize(FVectorVMCodeOptimizerContext& Context)
-	{
-		const uint32 SrcOpTypes = Context.BaseContext.DecodeSrcOperandTypes();
-		switch (SrcOpTypes)
-		{
-		case SRCOP_RRR: TBinaryKernelHandler<Kernel, DstHandler, RegisterHandler, RegisterHandler, NumInstancesPerOp>::Optimize(Context); break;
-		case SRCOP_RRC:	TBinaryKernelHandler<Kernel, DstHandler, ConstHandler, RegisterHandler, NumInstancesPerOp>::Optimize(Context); break;
-		case SRCOP_RCR: TBinaryKernelHandler<Kernel, DstHandler, RegisterHandler, ConstHandler, NumInstancesPerOp>::Optimize(Context); break;
-		case SRCOP_RCC:	TBinaryKernelHandler<Kernel, DstHandler, ConstHandler, ConstHandler, NumInstancesPerOp>::Optimize(Context); break;
-		default: check(0); break;
-		};
-	}
-
-	static void Exec(FVectorVMContext& Context)
-	{
-		const uint32 SrcOpTypes = Context.DecodeSrcOperandTypes();
-		switch (SrcOpTypes)
-		{
-		case SRCOP_RRR: TBinaryKernelHandler<Kernel, DstHandler, RegisterHandler, RegisterHandler, NumInstancesPerOp>::Exec(Context); break;
-		case SRCOP_RRC:	TBinaryKernelHandler<Kernel, DstHandler, ConstHandler, RegisterHandler, NumInstancesPerOp>::Exec(Context); break;
-		case SRCOP_RCR: TBinaryKernelHandler<Kernel, DstHandler, RegisterHandler, ConstHandler, NumInstancesPerOp>::Exec(Context); break;
-		case SRCOP_RCC:	TBinaryKernelHandler<Kernel, DstHandler, ConstHandler, ConstHandler, NumInstancesPerOp>::Exec(Context); break;
-		default: check(0); break;
-		};
-	}
-};
-template<typename Kernel>
-struct TBinaryScalarKernel : public TBinaryKernel<Kernel, FRegisterHandler<float>, FConstantHandler<float>, FRegisterHandler<float>, 1> {};
-template<typename Kernel>
-struct TBinaryVectorKernel : public TBinaryKernel<Kernel, FRegisterHandler<VectorRegister4Float>, FConstantHandler<VectorRegister4Float>, FRegisterHandler<VectorRegister4Float>, VECTOR_WIDTH_FLOATS> {};
-template<typename Kernel>
-struct TBinaryVectorIntKernel : public TBinaryKernel<Kernel, FRegisterHandler<VectorRegister4Int>, FConstantHandler<VectorRegister4Int>, FRegisterHandler<VectorRegister4Int>, VECTOR_WIDTH_FLOATS> {};
-
-/** Base class of Vector kernels with 3 operands. */
-template <typename Kernel, typename DstHandler, typename ConstHandler, typename RegisterHandler, uint32 NumInstancesPerOp>
-struct TTrinaryKernel
-{
-	static void Optimize(FVectorVMCodeOptimizerContext& Context)
-	{
-		const uint32 SrcOpTypes = Context.BaseContext.DecodeSrcOperandTypes();
-		switch (SrcOpTypes)
-		{
-		case SRCOP_RRR: TTrinaryKernelHandler<Kernel, DstHandler, RegisterHandler, RegisterHandler, RegisterHandler, NumInstancesPerOp>::Optimize(Context); break;
-		case SRCOP_RRC:	TTrinaryKernelHandler<Kernel, DstHandler, ConstHandler, RegisterHandler, RegisterHandler, NumInstancesPerOp>::Optimize(Context); break;
-		case SRCOP_RCR: TTrinaryKernelHandler<Kernel, DstHandler, RegisterHandler, ConstHandler, RegisterHandler, NumInstancesPerOp>::Optimize(Context); break;
-		case SRCOP_RCC:	TTrinaryKernelHandler<Kernel, DstHandler, ConstHandler, ConstHandler, RegisterHandler, NumInstancesPerOp>::Optimize(Context); break;
-		case SRCOP_CRR: TTrinaryKernelHandler<Kernel, DstHandler, RegisterHandler, RegisterHandler, ConstHandler, NumInstancesPerOp>::Optimize(Context); break;
-		case SRCOP_CRC:	TTrinaryKernelHandler<Kernel, DstHandler, ConstHandler, RegisterHandler, ConstHandler, NumInstancesPerOp>::Optimize(Context); break;
-		case SRCOP_CCR: TTrinaryKernelHandler<Kernel, DstHandler, RegisterHandler, ConstHandler, ConstHandler, NumInstancesPerOp>::Optimize(Context); break;
-		case SRCOP_CCC:	TTrinaryKernelHandler<Kernel, DstHandler, ConstHandler, ConstHandler, ConstHandler, NumInstancesPerOp>::Optimize(Context); break;
-		default: check(0); break;
-		};
-	}
-
-	static void Exec(FVectorVMContext& Context)
-	{
-		const uint32 SrcOpTypes = Context.DecodeSrcOperandTypes();
-		switch (SrcOpTypes)
-		{
-		case SRCOP_RRR: TTrinaryKernelHandler<Kernel, DstHandler, RegisterHandler, RegisterHandler, RegisterHandler, NumInstancesPerOp>::Exec(Context); break;
-		case SRCOP_RRC:	TTrinaryKernelHandler<Kernel, DstHandler, ConstHandler, RegisterHandler, RegisterHandler, NumInstancesPerOp>::Exec(Context); break;
-		case SRCOP_RCR: TTrinaryKernelHandler<Kernel, DstHandler, RegisterHandler, ConstHandler, RegisterHandler, NumInstancesPerOp>::Exec(Context); break;
-		case SRCOP_RCC:	TTrinaryKernelHandler<Kernel, DstHandler, ConstHandler, ConstHandler, RegisterHandler, NumInstancesPerOp>::Exec(Context); break;
-		case SRCOP_CRR: TTrinaryKernelHandler<Kernel, DstHandler, RegisterHandler, RegisterHandler, ConstHandler, NumInstancesPerOp>::Exec(Context); break;
-		case SRCOP_CRC:	TTrinaryKernelHandler<Kernel, DstHandler, ConstHandler, RegisterHandler, ConstHandler, NumInstancesPerOp>::Exec(Context); break;
-		case SRCOP_CCR: TTrinaryKernelHandler<Kernel, DstHandler, RegisterHandler, ConstHandler, ConstHandler, NumInstancesPerOp>::Exec(Context); break;
-		case SRCOP_CCC:	TTrinaryKernelHandler<Kernel, DstHandler, ConstHandler, ConstHandler, ConstHandler, NumInstancesPerOp>::Exec(Context); break;
-		default: check(0); break;
-		};
-	}
-};
-
-template<typename Kernel>
-struct TTrinaryScalarKernel : public TTrinaryKernel<Kernel, FRegisterHandler<float>, FConstantHandler<float>, FRegisterHandler<float>, 1> {};
-template<typename Kernel>
-struct TTrinaryVectorKernel : public TTrinaryKernel<Kernel, FRegisterHandler<VectorRegister4Float>, FConstantHandler<VectorRegister4Float>, FRegisterHandler<VectorRegister4Float>, VECTOR_WIDTH_FLOATS> {};
-template<typename Kernel>
-struct TTrinaryVectorIntKernel : public TTrinaryKernel<Kernel, FRegisterHandler<VectorRegister4Int>, FConstantHandler<VectorRegister4Int>, FRegisterHandler<VectorRegister4Int>, VECTOR_WIDTH_FLOATS> {};
-
-
-/*------------------------------------------------------------------------------
-	Implementation of all kernel operations.
-------------------------------------------------------------------------------*/
-
-struct FVectorKernelAdd : public TBinaryVectorKernel<FVectorKernelAdd>
-{
-	static void VM_FORCEINLINE DoKernel(FVectorVMContext& Context, VectorRegister4Float* RESTRICT Dst,VectorRegister4Float Src0,VectorRegister4Float Src1)
-	{
-		*Dst = VectorAdd(Src0, Src1);
-	}
-};
-
-struct FVectorKernelSub : public TBinaryVectorKernel<FVectorKernelSub>
-{
-	static void VM_FORCEINLINE DoKernel(FVectorVMContext& Context, VectorRegister4Float* RESTRICT Dst,VectorRegister4Float Src0,VectorRegister4Float Src1)
-	{
-		*Dst = VectorSubtract(Src0, Src1);
-	}
-};
-
-struct FVectorKernelMul : public TBinaryVectorKernel<FVectorKernelMul>
-{
-	static void VM_FORCEINLINE DoKernel(FVectorVMContext& Context, VectorRegister4Float* RESTRICT Dst,VectorRegister4Float Src0,VectorRegister4Float Src1)
-	{
-		*Dst = VectorMultiply(Src0, Src1);
-	}
-};
-
-struct FVectorKernelDiv : public TBinaryVectorKernel<FVectorKernelDiv>
-{
-	static void VM_FORCEINLINE DoKernel(FVectorVMContext& Context, VectorRegister4Float* RESTRICT Dst, VectorRegister4Float Src0, VectorRegister4Float Src1)
-	{
-		*Dst = VectorDivide(Src0, Src1);
-	}
-};
-
-struct FVectorKernelDivSafe : public TBinaryVectorKernel<FVectorKernelDivSafe>
-{
-	static void VM_FORCEINLINE DoKernel(FVectorVMContext& Context, VectorRegister4Float* RESTRICT Dst, VectorRegister4Float Src0, VectorRegister4Float Src1)
-	{
-		VectorRegister4Float ValidMask = VectorCompareGT(VectorAbs(Src1), GlobalVectorConstants::SmallNumber);
-		*Dst = VectorSelect(ValidMask, VectorDivide(Src0, Src1), GlobalVectorConstants::FloatZero);
-	}
-};
-
-struct FVectorKernelMad : public TTrinaryVectorKernel<FVectorKernelMad>
-{
-	static void VM_FORCEINLINE DoKernel(FVectorVMContext& Context, VectorRegister4Float* RESTRICT Dst,VectorRegister4Float Src0,VectorRegister4Float Src1,VectorRegister4Float Src2)
-	{
-		*Dst = VectorMultiplyAdd(Src0, Src1, Src2);
-	}
-};
-
-struct FVectorKernelLerp : public TTrinaryVectorKernel<FVectorKernelLerp>
-{
-	static void VM_FORCEINLINE DoKernel(FVectorVMContext& Context, VectorRegister4Float* RESTRICT Dst,VectorRegister4Float Src0,VectorRegister4Float Src1,VectorRegister4Float Src2)
-	{
-		const VectorRegister4Float OneMinusAlpha = VectorSubtract(GlobalVectorConstants::FloatOne, Src2);
-		const VectorRegister4Float Tmp = VectorMultiply(Src0, OneMinusAlpha);
-		*Dst = VectorMultiplyAdd(Src1, Src2, Tmp);
-	}
-};
-
-struct FVectorKernelRcp : public TUnaryVectorKernel<FVectorKernelRcp>
-{
-	static void VM_FORCEINLINE DoKernel(FVectorVMContext& Context, VectorRegister4Float* RESTRICT Dst,VectorRegister4Float Src0)
-	{
-		*Dst = VectorVMAccuracy::Reciprocal(Src0);
-	}
-};
-
-// if the magnitude of the value is too small, then the result will be 0 (not NaN/Inf)
-struct FVectorKernelRcpSafe : public TUnaryVectorKernel<FVectorKernelRcpSafe>
-{
-	static void VM_FORCEINLINE DoKernel(FVectorVMContext& Context, VectorRegister4Float* RESTRICT Dst, VectorRegister4Float Src0)
-	{
-		VectorRegister4Float ValidMask = VectorCompareGT(VectorAbs(Src0), GlobalVectorConstants::SmallNumber);
-		*Dst = VectorSelect(ValidMask, VectorVMAccuracy::Reciprocal(Src0), GlobalVectorConstants::FloatZero);
-	}
-};
-
-struct FVectorKernelRsq : public TUnaryVectorKernel<FVectorKernelRsq>
-{
-	static void VM_FORCEINLINE DoKernel(FVectorVMContext& Context, VectorRegister4Float* RESTRICT Dst,VectorRegister4Float Src0)
-	{
-		*Dst = VectorVMAccuracy::ReciprocalSqrt(Src0);
-	}
-};
-
-// if the value is very small or negative, then the result will be 0 (not NaN/Inf/imaginary)
-struct FVectorKernelRsqSafe : public TUnaryVectorKernel<FVectorKernelRsqSafe>
-{
-	static void VM_FORCEINLINE DoKernel(FVectorVMContext& Context, VectorRegister4Float* RESTRICT Dst, VectorRegister4Float Src0)
-	{
-		VectorRegister4Float ValidMask = VectorCompareGT(Src0, GlobalVectorConstants::SmallNumber);
-		*Dst = VectorSelect(ValidMask, VectorVMAccuracy::ReciprocalSqrt(Src0), GlobalVectorConstants::FloatZero);
-	}
-};
-
-struct FVectorKernelSqrt : public TUnaryVectorKernel<FVectorKernelSqrt>
-{
-	static void VM_FORCEINLINE DoKernel(FVectorVMContext& Context, VectorRegister4Float* RESTRICT Dst,VectorRegister4Float Src0)
-	{
-		*Dst = VectorVMAccuracy::Sqrt(Src0);
-	}
-};
-
-struct FVectorKernelSqrtSafe : public TUnaryVectorKernel<FVectorKernelSqrtSafe>
-{
-	static void VM_FORCEINLINE DoKernel(FVectorVMContext& Context, VectorRegister4Float* RESTRICT Dst, VectorRegister4Float Src0)
-	{
-		VectorRegister4Float ValidMask = VectorCompareGT(Src0, GlobalVectorConstants::SmallNumber);
-		*Dst = VectorSelect(ValidMask, VectorVMAccuracy::Sqrt(Src0), GlobalVectorConstants::FloatZero);
-	}
-};
-
-struct FVectorKernelNeg : public TUnaryVectorKernel<FVectorKernelNeg>
-{
-	static void VM_FORCEINLINE DoKernel(FVectorVMContext& Context, VectorRegister4Float* RESTRICT Dst,VectorRegister4Float Src0)
-	{
-		*Dst = VectorNegate(Src0);
-	}
-};
-
-struct FVectorKernelAbs : public TUnaryVectorKernel<FVectorKernelAbs>
-{
-	static void VM_FORCEINLINE DoKernel(FVectorVMContext& Context, VectorRegister4Float* RESTRICT Dst,VectorRegister4Float Src0)
-	{
-		*Dst = VectorAbs(Src0);
-	}
-};
-
-struct FVectorKernelExp : public TUnaryVectorKernel<FVectorKernelExp>
-{
-	static void VM_FORCEINLINE DoKernel(FVectorVMContext& Context, VectorRegister4Float* RESTRICT Dst, VectorRegister4Float Src0)
-	{
-		*Dst = VectorExp(Src0);
-	}
-}; 
-
-struct FVectorKernelExp2 : public TUnaryVectorKernel<FVectorKernelExp2>
-{
-	static void VM_FORCEINLINE DoKernel(FVectorVMContext& Context, VectorRegister4Float* RESTRICT Dst, VectorRegister4Float Src0)
-	{
-		*Dst = VectorExp2(Src0);
-	}
-};
-
-struct FVectorKernelLog : public TUnaryVectorKernel<FVectorKernelLog>
-{
-	static void VM_FORCEINLINE DoKernel(FVectorVMContext& Context, VectorRegister4Float* RESTRICT Dst, VectorRegister4Float Src0)
-	{
-		*Dst = VectorLog(Src0);
-	}
-};
-
-struct FVectorKernelLogSafe : public TUnaryVectorKernel<FVectorKernelLogSafe>
-{
-	static void VM_FORCEINLINE DoKernel(FVectorVMContext& Context, VectorRegister4Float* RESTRICT Dst, VectorRegister4Float Src0)
-	{
-		VectorRegister4Float ValidMask = VectorCompareGT(Src0, GlobalVectorConstants::FloatZero);
-
-		*Dst = VectorSelect(ValidMask, VectorLog(Src0), GlobalVectorConstants::FloatZero);
-	}
-};
-
-struct FVectorKernelLog2 : public TUnaryVectorKernel<FVectorKernelLog2>
-{
-	static void VM_FORCEINLINE DoKernel(FVectorVMContext& Context, VectorRegister4Float* RESTRICT Dst, VectorRegister4Float Src0)
-	{
-		*Dst = VectorLog2(Src0);
-	}
-};
-
-struct FVectorKernelStep : public TBinaryVectorKernel<FVectorKernelStep>
-{
-	static void VM_FORCEINLINE DoKernel(FVectorVMContext& Context, VectorRegister4Float* RESTRICT Dst, VectorRegister4Float Src0, VectorRegister4Float Src1)
-	{
-		*Dst = VectorStep(VectorSubtract(Src1, Src0));
-	}
-};
-
-struct FVectorKernelClamp : public TTrinaryVectorKernel<FVectorKernelClamp>
-{
-	static void VM_FORCEINLINE DoKernel(FVectorVMContext& Context, VectorRegister4Float* RESTRICT Dst,VectorRegister4Float Src0,VectorRegister4Float Src1,VectorRegister4Float Src2)
-	{
-		const VectorRegister4Float Tmp = VectorMax(Src0, Src1);
-		*Dst = VectorMin(Tmp, Src2);
-	}
-};
-
-struct FVectorKernelSin : public TUnaryVectorKernel<FVectorKernelSin>
-{
-	static void VM_FORCEINLINE DoKernel(FVectorVMContext& Context, VectorRegister4Float* RESTRICT Dst, VectorRegister4Float Src0)
-	{
-		*Dst = VectorSin(Src0);
-	}
-};
-
-struct FVectorKernelCos : public TUnaryVectorKernel<FVectorKernelCos>
-{
-	static void VM_FORCEINLINE DoKernel(FVectorVMContext& Context, VectorRegister4Float* RESTRICT Dst, VectorRegister4Float Src0)
- 	{
-		*Dst = VectorCos(Src0);
-	}
-};
-
-struct FVectorKernelTan : public TUnaryVectorKernel<FVectorKernelTan>
-{
-	static void VM_FORCEINLINE DoKernel(FVectorVMContext& Context, VectorRegister4Float* RESTRICT Dst, VectorRegister4Float Src0)
-	{
-		*Dst = VectorTan(Src0);
-	}
-};
-
-struct FVectorKernelASin : public TUnaryVectorKernel<FVectorKernelASin>
-{
-	static void VM_FORCEINLINE DoKernel(FVectorVMContext& Context, VectorRegister4Float* RESTRICT Dst, VectorRegister4Float Src0)
-	{
-		*Dst = VectorASin(Src0);
-	}
-};
-
-struct FVectorKernelACos : public TUnaryVectorKernel<FVectorKernelACos>
-{
-	static void VM_FORCEINLINE DoKernel(FVectorVMContext& Context, VectorRegister4Float* RESTRICT Dst, VectorRegister4Float Src0)
-	{
-		*Dst = VectorACos(Src0);
-	}
-};
-
-struct FVectorKernelATan : public TUnaryVectorKernel<FVectorKernelATan>
-{
-	static void VM_FORCEINLINE DoKernel(FVectorVMContext& Context, VectorRegister4Float* RESTRICT Dst, VectorRegister4Float Src0)
-	{
-		*Dst = VectorATan(Src0);
-	}
-};
-
-struct FVectorKernelATan2 : public TBinaryVectorKernel<FVectorKernelATan2>
-{
-	static void VM_FORCEINLINE DoKernel(FVectorVMContext& Context, VectorRegister4Float* RESTRICT Dst, VectorRegister4Float Src0, VectorRegister4Float Src1)
-	{
-		*Dst = VectorATan2(Src0, Src1);
-	}
-};
-
-struct FVectorKernelCeil : public TUnaryVectorKernel<FVectorKernelCeil>
-{
-	static void VM_FORCEINLINE DoKernel(FVectorVMContext& Context, VectorRegister4Float* RESTRICT Dst, VectorRegister4Float Src0)
-	{
-		*Dst = VectorCeil(Src0);
-	}
-};
-
-struct FVectorKernelFloor : public TUnaryVectorKernel<FVectorKernelFloor>
-{
-	static void VM_FORCEINLINE DoKernel(FVectorVMContext& Context, VectorRegister4Float* RESTRICT Dst, VectorRegister4Float Src0)
-	{
-		*Dst = VectorFloor(Src0);
-	}
-};
-
-struct FVectorKernelRound : public TUnaryVectorKernel<FVectorKernelRound>
-{
-	static void VM_FORCEINLINE DoKernel(FVectorVMContext& Context, VectorRegister4Float* RESTRICT Dst, VectorRegister4Float Src0)
-	{
-		//TODO: >SSE4 has direct ops for this.		
-		VectorRegister4Float Trunc = VectorTruncate(Src0);
-		*Dst = VectorAdd(Trunc, VectorTruncate(VectorMultiply(VectorSubtract(Src0, Trunc), GlobalVectorConstants::FloatAlmostTwo())));
-	}
-};
-
-struct FVectorKernelMod : public TBinaryVectorKernel<FVectorKernelMod>
-{
-	static void VM_FORCEINLINE DoKernel(FVectorVMContext& Context, VectorRegister4Float* RESTRICT Dst, VectorRegister4Float Src0, VectorRegister4Float Src1)
-	{
-		*Dst = VectorMod(Src0, Src1);
-	}
-};
-
-struct FVectorKernelFrac : public TUnaryVectorKernel<FVectorKernelFrac>
-{
-	static void VM_FORCEINLINE DoKernel(FVectorVMContext& Context, VectorRegister4Float* RESTRICT Dst, VectorRegister4Float Src0)
-	{
-		*Dst = VectorFractional(Src0);
-	}
-};
-
-struct FVectorKernelTrunc : public TUnaryVectorKernel<FVectorKernelTrunc>
-{
-	static void VM_FORCEINLINE DoKernel(FVectorVMContext& Context, VectorRegister4Float* RESTRICT Dst, VectorRegister4Float Src0)
-	{
-		*Dst = VectorTruncate(Src0);
-	}
-};
-
-struct FVectorKernelCompareLT : public TBinaryVectorKernel<FVectorKernelCompareLT>
-{
-	static void VM_FORCEINLINE DoKernel(FVectorVMContext& Context, VectorRegister4Float* RESTRICT Dst, VectorRegister4Float Src0, VectorRegister4Float Src1)
-	{
-		*Dst = VectorCompareLT(Src0, Src1);
-	}
-};
-
-struct FVectorKernelCompareLE : public TBinaryVectorKernel<FVectorKernelCompareLE>
-{
-	static void VM_FORCEINLINE DoKernel(FVectorVMContext& Context, VectorRegister4Float* RESTRICT Dst, VectorRegister4Float Src0, VectorRegister4Float Src1)
-	{
-		*Dst = VectorCompareLE(Src0, Src1);
-	}
-};
-
-struct FVectorKernelCompareGT : public TBinaryVectorKernel<FVectorKernelCompareGT>
-{
-	static void VM_FORCEINLINE DoKernel(FVectorVMContext& Context, VectorRegister4Float* RESTRICT Dst, VectorRegister4Float Src0, VectorRegister4Float Src1)
-	{
-		*Dst = VectorCompareGT(Src0, Src1);
-	}
-};
-
-struct FVectorKernelCompareGE : public TBinaryVectorKernel<FVectorKernelCompareGE>
-{
-	static void VM_FORCEINLINE DoKernel(FVectorVMContext& Context, VectorRegister4Float* RESTRICT Dst, VectorRegister4Float Src0, VectorRegister4Float Src1)
-	{
-		*Dst = VectorCompareGE(Src0, Src1);
-	}
-};
-
-struct FVectorKernelCompareEQ : public TBinaryVectorKernel<FVectorKernelCompareEQ>
-{
-	static void VM_FORCEINLINE DoKernel(FVectorVMContext& Context, VectorRegister4Float* RESTRICT Dst, VectorRegister4Float Src0, VectorRegister4Float Src1)
-	{
-		*Dst = VectorCompareEQ(Src0, Src1);
-	}
-};
-
-struct FVectorKernelCompareNEQ : public TBinaryVectorKernel<FVectorKernelCompareNEQ>
-{
-	static void VM_FORCEINLINE DoKernel(FVectorVMContext& Context, VectorRegister4Float* RESTRICT Dst, VectorRegister4Float Src0, VectorRegister4Float Src1)
-	{		
-		*Dst = VectorCompareNE(Src0, Src1);
-	}
-};
-
-struct FVectorKernelSelect : public TTrinaryVectorKernel<FVectorKernelSelect>
-{
-	static void VM_FORCEINLINE DoKernel(FVectorVMContext& Context, VectorRegister4Float* RESTRICT Dst, VectorRegister4Float Mask, VectorRegister4Float A, VectorRegister4Float B)
-	{
-		*Dst = VectorSelect(Mask, A, B);
-	}
-};
-
-struct FVectorKernelExecutionIndex
-{
-	static void Optimize(FVectorVMCodeOptimizerContext& Context)
-	{
-		Context.WriteExecFunction(Exec);
-		FRegisterHandler<VectorRegister4Int>::Optimize(Context);
-	}
-
-	static void VM_FORCEINLINE Exec(FVectorVMContext& Context)
-	{
-		static_assert(VECTOR_WIDTH_FLOATS == 4, "Need to update this when upgrading the VM to support >SSE2");
-		VectorRegister4Int VectorStride = MakeVectorRegisterInt(VECTOR_WIDTH_FLOATS, VECTOR_WIDTH_FLOATS, VECTOR_WIDTH_FLOATS, VECTOR_WIDTH_FLOATS);
-		VectorRegister4Int Index = MakeVectorRegisterInt(Context.StartInstance, Context.StartInstance + 1, Context.StartInstance + 2, Context.StartInstance + 3);
-		
-		FRegisterHandler<VectorRegister4Int> Dest(Context);
-		const int32 Loops = Context.GetNumLoops<VECTOR_WIDTH_FLOATS>();
-		for (int32 i = 0; i < Loops; ++i)
-		{
-			*Dest.GetDestAndAdvance() = Index;
-			Index = VectorIntAdd(Index, VectorStride);
-		}
-	}
-};
-
-struct FVectorKernelEnterStatScope
-{
-	static void Optimize(FVectorVMCodeOptimizerContext& Context)
-	{
-#if STATS
-		Context.WriteExecFunction(Exec);
-		FConstantHandler<int32>::Optimize(Context);
-#elif ENABLE_STATNAMEDEVENTS
-		if ( GbDetailedVMScriptStats )
-		{
-			Context.WriteExecFunction(Exec);
-			FConstantHandler<int32>::Optimize(Context);
-		}
-		else
-		{
-			FConstantHandler<int32>::OptimizeSkip(Context);
-		}
-#else
-		// just skip the op if we don't have stats enabled
-		FConstantHandler<int32>::OptimizeSkip(Context);
-#endif
-	}
-
-	static VM_FORCEINLINE void Exec(FVectorVMContext& Context)
-	{
-		FConstantHandler<int32> ScopeIdx(Context);
-#if STATS
-		if (GbDetailedVMScriptStats && Context.StatScopes.Num())
-		{
-			int32 CounterIdx = Context.StatCounterStack.AddDefaulted(1);
-			int32 ScopeIndex = ScopeIdx.Get();
-			Context.StatCounterStack[CounterIdx].CycleCounter.Start(Context.StatScopes[ScopeIndex].StatId);
-			Context.StatCounterStack[CounterIdx].VmCycleCounter = { ScopeIndex, FPlatformTime::Cycles64() };
-		}
-#elif ENABLE_STATNAMEDEVENTS
-		if (Context.StatNamedEventScopes.Num())
-		{
-			FPlatformMisc::BeginNamedEvent(FColor::Red, *Context.StatNamedEventScopes[ScopeIdx.Get()]);
-		}
-#endif
-	}
-};
-
-struct FVectorKernelExitStatScope
-{
-	static void Optimize(FVectorVMCodeOptimizerContext& Context)
-	{
-#if STATS
-		Context.WriteExecFunction(Exec);
-#elif ENABLE_STATNAMEDEVENTS
-		if (GbDetailedVMScriptStats)
-		{
-			Context.WriteExecFunction(Exec);
-		}
-#endif
-	}
-		
-	static VM_FORCEINLINE void Exec(FVectorVMContext& Context)
-	{
-#if STATS
-		if (GbDetailedVMScriptStats && Context.StatScopes.Num())
-		{
-			FStatStackEntry& StackEntry = Context.StatCounterStack.Last();
-			StackEntry.CycleCounter.Stop();
-			Context.ScopeExecCycles[StackEntry.VmCycleCounter.ScopeIndex] += FPlatformTime::Cycles64() - StackEntry.VmCycleCounter.ScopeEnterCycles;
-			Context.StatCounterStack.Pop(EAllowShrinking::No);
-		}
-#elif ENABLE_STATNAMEDEVENTS
-		if (Context.StatNamedEventScopes.Num())
-		{
-			FPlatformMisc::EndNamedEvent();
-		}
-#endif
-	}
-};
-
-struct FVectorKernelRandom : public TUnaryVectorKernel<FVectorKernelRandom>
-{
-	static void VM_FORCEINLINE DoKernel(FVectorVMContext& Context, VectorRegister4Float* RESTRICT Dst, VectorRegister4Float Src0)
-	{
-		//EEK!. Improve this. Implement GPU style seeded rand instead of this.
-		VectorRegister4Float Result = MakeVectorRegister(Context.RandStream.GetFraction(),
-			Context.RandStream.GetFraction(),
-			Context.RandStream.GetFraction(),
-			Context.RandStream.GetFraction());
-		*Dst = VectorMultiply(Result, Src0);
-	}
-};
-
-/* gaussian distribution random number (not working yet) */
-struct FVectorKernelRandomGauss : public TBinaryVectorKernel<FVectorKernelRandomGauss>
-{
-	static void VM_FORCEINLINE DoKernel(FVectorVMContext& Context, VectorRegister4Float* RESTRICT Dst, VectorRegister4Float Src0, VectorRegister4Float Src1)
-	{
-		VectorRegister4Float Result = MakeVectorRegister(Context.RandStream.GetFraction(),
-			Context.RandStream.GetFraction(),
-			Context.RandStream.GetFraction(),
-			Context.RandStream.GetFraction());
-
-		Result = VectorSubtract(Result, GlobalVectorConstants::FloatOneHalf);
-		Result = VectorMultiply(MakeVectorRegister(3.0f, 3.0f, 3.0f, 3.0f), Result);
-
-		// taylor series gaussian approximation
-		const VectorRegister4Float SPi2 = VectorReciprocal(VectorReciprocalSqrt(MakeVectorRegister(2 * PI, 2 * PI, 2 * PI, 2 * PI)));
-		VectorRegister4Float Gauss = VectorReciprocal(SPi2);
-		VectorRegister4Float Div = VectorMultiply(GlobalVectorConstants::FloatTwo, SPi2);
-		Gauss = VectorSubtract(Gauss, VectorDivide(VectorMultiply(Result, Result), Div));
-		Div = VectorMultiply(MakeVectorRegister(8.0f, 8.0f, 8.0f, 8.0f), SPi2);
-		Gauss = VectorAdd(Gauss, VectorDivide(VectorPow(MakeVectorRegister(4.0f, 4.0f, 4.0f, 4.0f), Result), Div));
-		Div = VectorMultiply(MakeVectorRegister(48.0f, 48.0f, 48.0f, 48.0f), SPi2);
-		Gauss = VectorSubtract(Gauss, VectorDivide(VectorPow(MakeVectorRegister(6.0f, 6.0f, 6.0f, 6.0f), Result), Div));
-
-		Gauss = VectorDivide(Gauss, MakeVectorRegister(0.4f, 0.4f, 0.4f, 0.4f));
-		Gauss = VectorMultiply(Gauss, Src0);
-		*Dst = Gauss;
-	}
-};
-
-struct FVectorKernelMin : public TBinaryVectorKernel<FVectorKernelMin>
-{
-	static void VM_FORCEINLINE DoKernel(FVectorVMContext& Context, VectorRegister4Float* RESTRICT Dst,VectorRegister4Float Src0,VectorRegister4Float Src1)
-	{
-		*Dst = VectorMin(Src0, Src1);
-	}
-};
-
-struct FVectorKernelMax : public TBinaryVectorKernel<FVectorKernelMax>
-{
-	static void VM_FORCEINLINE DoKernel(FVectorVMContext& Context, VectorRegister4Float* RESTRICT Dst,VectorRegister4Float Src0,VectorRegister4Float Src1)
-	{
-		*Dst = VectorMax(Src0, Src1);
-	}
-};
-
-struct FVectorKernelPow : public TBinaryVectorKernel<FVectorKernelPow>
-{
-	static void VM_FORCEINLINE DoKernel(FVectorVMContext& Context, VectorRegister4Float* RESTRICT Dst,VectorRegister4Float Src0,VectorRegister4Float Src1)
-	{
-		*Dst = VectorPow(Src0, Src1);
-	}
-};
-
-// if the base is small, then the result will be 0
-struct FVectorKernelPowSafe : public TBinaryVectorKernel<FVectorKernelPowSafe>
-{
-	static void VM_FORCEINLINE DoKernel(FVectorVMContext& Context, VectorRegister4Float* RESTRICT Dst, VectorRegister4Float Src0, VectorRegister4Float Src1)
-	{
-		VectorRegister4Float ValidMask = VectorCompareGT(Src0, GlobalVectorConstants::SmallNumber);
-		*Dst = VectorSelect(ValidMask, VectorPow(Src0, Src1), GlobalVectorConstants::FloatZero);
-	}
-};
-
-struct FVectorKernelSign : public TUnaryVectorKernel<FVectorKernelSign>
-{
-	static void VM_FORCEINLINE DoKernel(FVectorVMContext& Context, VectorRegister4Float* RESTRICT Dst, VectorRegister4Float Src0)
-	{
-		*Dst = VectorSign(Src0);
-	}
-};
-
-namespace VectorVMNoise
-{
-	int32 P[512] =
-	{
-		151,160,137,91,90,15,
-		131,13,201,95,96,53,194,233,7,225,140,36,103,30,69,142,8,99,37,240,21,10,23,
-		190, 6,148,247,120,234,75,0,26,197,62,94,252,219,203,117,35,11,32,57,177,33,
-		88,237,149,56,87,174,20,125,136,171,168, 68,175,74,165,71,134,139,48,27,166,
-		77,146,158,231,83,111,229,122,60,211,133,230,220,105,92,41,55,46,245,40,244,
-		102,143,54, 65,25,63,161, 1,216,80,73,209,76,132,187,208, 89,18,169,200,196,
-		135,130,116,188,159,86,164,100,109,198,173,186, 3,64,52,217,226,250,124,123,
-		5,202,38,147,118,126,255,82,85,212,207,206,59,227,47,16,58,17,182,189,28,42,
-		223,183,170,213,119,248,152, 2,44,154,163, 70,221,153,101,155,167, 43,172,9,
-		129,22,39,253, 19,98,108,110,79,113,224,232,178,185, 112,104,218,246,97,228,
-		251,34,242,193,238,210,144,12,191,179,162,241, 81,51,145,235,249,14,239,107,
-		49,192,214, 31,181,199,106,157,184, 84,204,176,115,121,50,45,127, 4,150,254,
-		138,236,205,93,222,114,67,29,24,72,243,141,128,195,78,66,215,61,156,180,
-		151,160,137,91,90,15,
-		131,13,201,95,96,53,194,233,7,225,140,36,103,30,69,142,8,99,37,240,21,10,23,
-		190, 6,148,247,120,234,75,0,26,197,62,94,252,219,203,117,35,11,32,57,177,33,
-		88,237,149,56,87,174,20,125,136,171,168, 68,175,74,165,71,134,139,48,27,166,
-		77,146,158,231,83,111,229,122,60,211,133,230,220,105,92,41,55,46,245,40,244,
-		102,143,54, 65,25,63,161, 1,216,80,73,209,76,132,187,208, 89,18,169,200,196,
-		135,130,116,188,159,86,164,100,109,198,173,186, 3,64,52,217,226,250,124,123,
-		5,202,38,147,118,126,255,82,85,212,207,206,59,227,47,16,58,17,182,189,28,42,
-		223,183,170,213,119,248,152, 2,44,154,163, 70,221,153,101,155,167, 43,172,9,
-		129,22,39,253, 19,98,108,110,79,113,224,232,178,185, 112,104,218,246,97,228,
-		251,34,242,193,238,210,144,12,191,179,162,241, 81,51,145,235,249,14,239,107,
-		49,192,214, 31,181,199,106,157,184, 84,204,176,115,121,50,45,127, 4,150,254,
-		138,236,205,93,222,114,67,29,24,72,243,141,128,195,78,66,215,61,156,180
-	};
-
-	static FORCEINLINE float Lerp(float X, float A, float B)
-	{
-		return A + X * (B - A);
-	}
-
-	static FORCEINLINE float Fade(float X)
-	{
-		return X * X * X * (X * (X * 6 - 15) + 10);
-	}
-	
-	static FORCEINLINE float Grad(int32 hash, float x, float y, float z)
-	{
-		 hash &= 15;
-		 float u = (hash < 8) ? x : y;
-		 float v = (hash < 4) ? y : ((hash == 12 || hash == 14) ? x : z);
-		 return ((hash & 1) == 0 ? u : -u) + ((hash & 2) == 0 ? v : -v);
-	}
-
-	struct FScalarKernelNoise3D_iNoise : TTrinaryScalarKernel<FScalarKernelNoise3D_iNoise>
-	{
-		static void VM_FORCEINLINE DoKernel(FVectorVMContext& Context, float* RESTRICT Dst, float X, float Y, float Z)
-		{
-			float Xfl = FMath::FloorToFloat(X);
-			float Yfl = FMath::FloorToFloat(Y);
-			float Zfl = FMath::FloorToFloat(Z);
-			int32 Xi = (int32)(Xfl) & 255;
-			int32 Yi = (int32)(Yfl) & 255;
-			int32 Zi = (int32)(Zfl) & 255;
-			X -= Xfl;
-			Y -= Yfl;
-			Z -= Zfl;
-			float Xm1 = X - 1.0f;
-			float Ym1 = Y - 1.0f;
-			float Zm1 = Z - 1.0f;
-
-			int32 A = P[Xi] + Yi;
-			int32 AA = P[A] + Zi;	int32 AB = P[A + 1] + Zi;
-
-			int32 B = P[Xi + 1] + Yi;
-			int32 BA = P[B] + Zi;	int32 BB = P[B + 1] + Zi;
-
-			float U = Fade(X);
-			float V = Fade(Y);
-			float W = Fade(Z);
-
-			*Dst =
-				Lerp(W,
-					Lerp(V,
-						Lerp(U,
-							Grad(P[AA], X, Y, Z),
-							Grad(P[BA], Xm1, Y, Z)),
-						Lerp(U,
-							Grad(P[AB], X, Ym1, Z),
-							Grad(P[BB], Xm1, Ym1, Z))),
-					Lerp(V,
-						Lerp(U,
-							Grad(P[AA + 1], X, Y, Zm1),
-							Grad(P[BA + 1], Xm1, Y, Zm1)),
-						Lerp(U,
-							Grad(P[AB + 1], X, Ym1, Zm1),
-							Grad(P[BB + 1], Xm1, Ym1, Zm1))));
-		}
-	};
-
-	struct FScalarKernelNoise2D_iNoise : TBinaryScalarKernel<FScalarKernelNoise2D_iNoise>
-	{
-		static void VM_FORCEINLINE DoKernel(FVectorVMContext& Context, float* RESTRICT Dst, float X, float Y)
-		{
-			*Dst = 0.0f;//TODO
-		}
-	};
-
-	struct FScalarKernelNoise1D_iNoise : TUnaryScalarKernel<FScalarKernelNoise1D_iNoise>
-	{
-		static void VM_FORCEINLINE DoKernel(FVectorVMContext& Context, float* RESTRICT Dst, float X)
-		{
-			*Dst = 0.0f;//TODO;
-		}
-	};
-
-	static void Noise1D(FVectorVMContext& Context) { FScalarKernelNoise1D_iNoise::Exec(Context); }
-	static void Noise2D(FVectorVMContext& Context) { FScalarKernelNoise2D_iNoise::Exec(Context); }
-	static void Noise3D(FVectorVMContext& Context)
-	{
-		//Basic scalar implementation of perlin's improved noise until I can spend some quality time exploring vectorized implementations of Marc O's noise from Random.ush.
-		//http://mrl.nyu.edu/~perlin/noise/
-		FScalarKernelNoise3D_iNoise::Exec(Context);
-	}
-
-	static void Optimize_Noise1D(FVectorVMCodeOptimizerContext& Context) { FScalarKernelNoise1D_iNoise::Optimize(Context); }
-	static void Optimize_Noise2D(FVectorVMCodeOptimizerContext& Context) { FScalarKernelNoise2D_iNoise::Optimize(Context); }
-	static void Optimize_Noise3D(FVectorVMCodeOptimizerContext& Context) { FScalarKernelNoise3D_iNoise::Optimize(Context); }
-};
-
-//Olaf's orginal curl noise. Needs updating for the new scalar VM and possibly calling Curl Noise to avoid confusion with regular noise?
-//Possibly needs to be a data interface as the VM can't output Vectors?
-struct FVectorKernelNoise : public TUnaryVectorKernel<FVectorKernelNoise>
-{
-	static VectorRegister4Float RandomTable[17][17][17];
-
-	static void VM_FORCEINLINE DoKernel(FVectorVMContext& Context, VectorRegister4Float* RESTRICT Dst, VectorRegister4Float Src0)
-	{
-		const VectorRegister4Float VecSize = MakeVectorRegister(16.0f, 16.0f, 16.0f, 16.0f);
-
-		*Dst = GlobalVectorConstants::FloatZero;
-		
-		for (uint32 i = 1; i < 2; i++)
-		{
-			float Di = 0.2f * (1.0f/(1<<i));
-			VectorRegister4Float Div = MakeVectorRegister(Di, Di, Di, Di);
-			VectorRegister4Float Coords = VectorMod( VectorAbs( VectorMultiply(Src0, Div) ), VecSize );
-			const float *CoordPtr = reinterpret_cast<float const*>(&Coords);
-			const int32 Cx = CoordPtr[0];
-			const int32 Cy = CoordPtr[1];
-			const int32 Cz = CoordPtr[2];
-
-			VectorRegister4Float Frac = VectorFractional(Coords);
-			VectorRegister4Float Alpha = VectorReplicate(Frac, 0);
-			VectorRegister4Float OneMinusAlpha = VectorSubtract(GlobalVectorConstants::FloatOne, Alpha);
-			
-			VectorRegister4Float XV1 = VectorMultiplyAdd(RandomTable[Cx][Cy][Cz], Alpha, VectorMultiply(RandomTable[Cx+1][Cy][Cz], OneMinusAlpha));
-			VectorRegister4Float XV2 = VectorMultiplyAdd(RandomTable[Cx][Cy+1][Cz], Alpha, VectorMultiply(RandomTable[Cx+1][Cy+1][Cz], OneMinusAlpha));
-			VectorRegister4Float XV3 = VectorMultiplyAdd(RandomTable[Cx][Cy][Cz+1], Alpha, VectorMultiply(RandomTable[Cx+1][Cy][Cz+1], OneMinusAlpha));
-			VectorRegister4Float XV4 = VectorMultiplyAdd(RandomTable[Cx][Cy+1][Cz+1], Alpha, VectorMultiply(RandomTable[Cx+1][Cy+1][Cz+1], OneMinusAlpha));
-
-			Alpha = VectorReplicate(Frac, 1);
-			OneMinusAlpha = VectorSubtract(GlobalVectorConstants::FloatOne, Alpha);
-			VectorRegister4Float YV1 = VectorMultiplyAdd(XV1, Alpha, VectorMultiply(XV2, OneMinusAlpha));
-			VectorRegister4Float YV2 = VectorMultiplyAdd(XV3, Alpha, VectorMultiply(XV4, OneMinusAlpha));
-
-			Alpha = VectorReplicate(Frac, 2);
-			OneMinusAlpha = VectorSubtract(GlobalVectorConstants::FloatOne, Alpha);
-			VectorRegister4Float ZV = VectorMultiplyAdd(YV1, Alpha, VectorMultiply(YV2, OneMinusAlpha));
-
-			*Dst = VectorAdd(*Dst, ZV);
-		}
-	}
-};
-
-namespace VectorKernelNoiseImpl
-{
-	static void BuildNoiseTable()
-	{
-		// random noise
-		float TempTable[17][17][17];
-		for (int z = 0; z < 17; z++)
-		{
-			for (int y = 0; y < 17; y++)
-			{
-				for (int x = 0; x < 17; x++)
-				{
-					float f1 = (float)FMath::FRandRange(-1.0f, 1.0f);
-					TempTable[x][y][z] = f1;
-				}
-			}
-		}
-
-		// pad
-		for (int i = 0; i < 17; i++)
-		{
-			for (int j = 0; j < 17; j++)
-			{
-				TempTable[i][j][16] = TempTable[i][j][0];
-				TempTable[i][16][j] = TempTable[i][0][j];
-				TempTable[16][j][i] = TempTable[0][j][i];
-			}
-		}
-
-		// compute gradients
-		FVector3f TempTable2[17][17][17];
-		for (int z = 0; z < 16; z++)
-		{
-			for (int y = 0; y < 16; y++)
-			{
-				for (int x = 0; x < 16; x++)
-				{
-					FVector3f XGrad = FVector3f(1.0f, 0.0f, TempTable[x][y][z] - TempTable[x + 1][y][z]);
-					FVector3f YGrad = FVector3f(0.0f, 1.0f, TempTable[x][y][z] - TempTable[x][y + 1][z]);
-					FVector3f ZGrad = FVector3f(0.0f, 1.0f, TempTable[x][y][z] - TempTable[x][y][z + 1]);
-
-					FVector3f Grad = FVector3f(XGrad.Z, YGrad.Z, ZGrad.Z);
-					TempTable2[x][y][z] = Grad;
-				}
-			}
-		}
-
-		// pad
-		for (int i = 0; i < 17; i++)
-		{
-			for (int j = 0; j < 17; j++)
-			{
-				TempTable2[i][j][16] = TempTable2[i][j][0];
-				TempTable2[i][16][j] = TempTable2[i][0][j];
-				TempTable2[16][j][i] = TempTable2[0][j][i];
-			}
-		}
-
-
-		// compute curl of gradient field
-		for (int z = 0; z < 16; z++)
-		{
-			for (int y = 0; y < 16; y++)
-			{
-				for (int x = 0; x < 16; x++)
-				{
-					FVector3f Dy = TempTable2[x][y][z] - TempTable2[x][y + 1][z];
-					FVector3f Sy = TempTable2[x][y][z] + TempTable2[x][y + 1][z];
-					FVector3f Dx = TempTable2[x][y][z] - TempTable2[x + 1][y][z];
-					FVector3f Sx = TempTable2[x][y][z] + TempTable2[x + 1][y][z];
-					FVector3f Dz = TempTable2[x][y][z] - TempTable2[x][y][z + 1];
-					FVector3f Sz = TempTable2[x][y][z] + TempTable2[x][y][z + 1];
-					FVector3f Dir = FVector3f(Dy.Z - Sz.Y, Dz.X - Sx.Z, Dx.Y - Sy.X);
-
-					FVectorKernelNoise::RandomTable[x][y][z] = MakeVectorRegister(Dir.X, Dir.Y, Dir.Z, 0.f);
-				}
-			}
-		}
-	}
-}
-
-VectorRegister4Float FVectorKernelNoise::RandomTable[17][17][17];
-
-//////////////////////////////////////////////////////////////////////////
-//Special Kernels.
-
-/** Special kernel for acquiring a new ID. TODO. Can be written as general RWBuffer ops when we support that. */
-struct FScalarKernelAcquireID
-{
-	static void Optimize(FVectorVMCodeOptimizerContext& Context)
-	{
-		Context.WriteExecFunction(Exec);
-		Context.Write(Context.DecodeU16());		// DataSetIndex
-		Context.Write(Context.DecodeU16());		// IDIndexReg
-		Context.Write(Context.DecodeU16());		// IDTagReg
-	}
-
-	static VM_FORCEINLINE void Exec(FVectorVMContext& Context)
-	{
-		const int32 DataSetIndex = Context.DecodeU16();
-		const TArrayView<FDataSetMeta> MetaTable = Context.DataSetMetaTable;
-		TArray<int32>&RESTRICT FreeIDTable = *MetaTable[DataSetIndex].FreeIDTable;
-		TArray<int32>&RESTRICT SpawnedIDsTable = *MetaTable[DataSetIndex].SpawnedIDsTable;
-
-		const int32 Tag = MetaTable[DataSetIndex].IDAcquireTag;
-
-		const int32 IDIndexReg = Context.DecodeU16();
-		int32*RESTRICT IDIndex = (int32*)(Context.GetTempRegister(IDIndexReg));
-
-		const int32 IDTagReg = Context.DecodeU16();
-		int32*RESTRICT IDTag = (int32*)(Context.GetTempRegister(IDTagReg));
-
-		int32& NumFreeIDs = *MetaTable[DataSetIndex].NumFreeIDs;
-
-		//Temporarily using a lock to ensure thread safety for accessing the FreeIDTable until a lock free solution can be implemented.
-		MetaTable[DataSetIndex].LockFreeTable();
-	
-		check(FreeIDTable.Num() >= Context.NumInstances);
-		check(NumFreeIDs >= Context.NumInstances);
-		for (int32 i = 0; i < Context.NumInstances; ++i)
-		{
-			int32 FreeIDTableIndex = --NumFreeIDs;
-
-			//Grab the value from the FreeIDTable.
-			int32 AcquiredID = FreeIDTable[FreeIDTableIndex];
-			checkSlow(AcquiredID != INDEX_NONE);
-
-			//UE_LOG(LogVectorVM, Warning, TEXT("AcquireID: ID:%d | FreeTableIdx:%d."), AcquiredID, FreeIDTableIndex);
-			//Mark this entry in the FreeIDTable as invalid.
-			FreeIDTable[FreeIDTableIndex] = INDEX_NONE;
-
-			*IDIndex = AcquiredID;
-			*IDTag = Tag;
-			++IDIndex;
-			++IDTag;
-
-			SpawnedIDsTable.Add(AcquiredID);
-		}
-
-		MetaTable[DataSetIndex].UnlockFreeTable();
-	}
-};
-
-/** Special kernel for updating a new ID. TODO. Can be written as general RWBuffer ops when we support that. */
-struct FScalarKernelUpdateID
-{
-	static void Optimize(FVectorVMCodeOptimizerContext& Context)
-	{
-		Context.WriteExecFunction(Exec);
-		Context.Write(Context.DecodeU16());		// DataSetIndex
-		Context.Write(Context.DecodeU16());		// InstanceIDRegisterIndex
-		Context.Write(Context.DecodeU16());		// InstanceIndexRegisterIndex
-	}
-
-	static VM_FORCEINLINE void Exec(FVectorVMContext& Context)
-	{
-		const int32 DataSetIndex = Context.DecodeU16();
-		const int32 InstanceIDRegisterIndex = Context.DecodeU16();
-		const int32 InstanceIndexRegisterIndex = Context.DecodeU16();
-
-		const TArrayView<FDataSetMeta> MetaTable = Context.DataSetMetaTable;
-
-		TArray<int32>&RESTRICT IDTable = *MetaTable[DataSetIndex].IDTable;
-		const int32 InstanceOffset = MetaTable[DataSetIndex].InstanceOffset;
-		const int32 AbsoluteStartInstance = InstanceOffset + Context.StartInstance;
-
-		const int32*RESTRICT IDRegister = (int32*)(Context.GetTempRegister(InstanceIDRegisterIndex));
-		const int32*RESTRICT IndexRegister = (int32*)(Context.GetTempRegister(InstanceIndexRegisterIndex));
-		
-		FDataSetThreadLocalTempData& DataSetTempData = Context.ThreadLocalTempData[DataSetIndex];
-
-		TArray<int32>&RESTRICT IDsToFree = DataSetTempData.IDsToFree;
-		check(IDTable.Num() >= AbsoluteStartInstance + Context.NumInstances);
-		for (int32 i = 0; i < Context.NumInstances; ++i)
-		{
-			int32 InstanceId = IDRegister[i];
-			int32 Index = IndexRegister[i];
-
-			if (Index == INDEX_NONE)
-			{
-				//Add the ID to a thread local list of IDs to free which are actually added to the list safely at the end of this chunk's execution.
-				IDsToFree.Add(InstanceId);
-				//UE_LOG(LogVectorVM, Warning, TEXT("FreeingID: InstanceID:%d."), InstanceId);
-			}
-			else
-			{
-				// Update the actual index for this ID. No thread safety is needed as this ID slot can only ever be written by this instance and so a single thread.
-				// The index passed into this function is the same as that given to the OutputData*() functions (FScalarKernelWriteOutputIndexed kernel).
-				// That value is local to the execution step (update or spawn), so we must offset it by the start instance number for the step, just like
-				// GetOutputRegister() does.
-				IDTable[InstanceId] = Index + InstanceOffset;
-
-				//Update thread local max ID seen. We push this to the real value at the end of execution.
-				DataSetTempData.MaxID = FMath::Max(DataSetTempData.MaxID, InstanceId);
-				
-				//UE_LOG(LogVectorVM, Warning, TEXT("UpdateID: RealIdx:%d | InstanceID:%d."), RealIdx, InstanceId);
-			}
-		}
-	}
-};
-
-/** Special kernel for reading from the main input dataset. */
-template<typename SourceType, int TypeOffset>
-struct FVectorKernelReadInput
-{
-	static void Optimize(FVectorVMCodeOptimizerContext& Context)
-	{
-		Context.WriteExecFunction(Exec);
-		Context.Write(Context.DecodeU16());	// DataSetIndex
-		Context.Write(Context.DecodeU16());	// InputRegisterIdx
-		Context.Write(Context.DecodeU16());	// DestRegisterIdx
-	}
-
-	static VM_FORCEINLINE void Exec(FVectorVMContext& Context)
-	{
-		static const int32 InstancesPerVector = sizeof(VectorRegister4Float ) / sizeof(SourceType);
-
-		const int32 DataSetIndex = Context.DecodeU16();
-		const int32 InputRegisterIdx = Context.DecodeU16();
-		const int32 DestRegisterIdx = Context.DecodeU16();
-		int32 Loops = Context.GetNumLoops<InstancesPerVector>();
-
-		VectorRegister4Float* DestReg = (VectorRegister4Float*)(Context.GetTempRegister(DestRegisterIdx));
-		VectorRegister4Float* InputReg = (VectorRegister4Float*)(Context.GetInputRegister<SourceType, TypeOffset>(DataSetIndex, InputRegisterIdx) + Context.GetStartInstance());
-
-		//TODO: We can actually do some scalar loads into the first and final vectors to get around alignment issues and then use the aligned load for all others.
-		while(Loops > 0)
-		{
-			*DestReg = VectorLoad(InputReg);
-			++DestReg;
-			++InputReg;
-			Loops--;
-		}
-	}
-};
-
-template<>
-struct FVectorKernelReadInput<FFloat16, 2>
-{
-	static void Optimize(FVectorVMCodeOptimizerContext& Context)
-	{
-		Context.WriteExecFunction(Exec);
-		Context.Write(Context.DecodeU16());	// DataSetIndex
-		Context.Write(Context.DecodeU16());	// InputRegisterIdx
-		Context.Write(Context.DecodeU16());	// DestRegisterIdx
-	}
-
-	static VM_FORCEINLINE void Exec(FVectorVMContext& Context)
-	{
-		const int32 DataSetIndex = Context.DecodeU16();
-		const int32 InputRegisterIdx = Context.DecodeU16();
-		const int32 DestRegisterIdx = Context.DecodeU16();
-		int32 Loops = Context.GetNumLoops<4>();
-
-		float* DestReg = (float*)(Context.GetTempRegister(DestRegisterIdx));
-		uint16* InputReg = (uint16*)(Context.GetInputRegister<FFloat16, 2>(DataSetIndex, InputRegisterIdx) + Context.GetStartInstance());
-
-		//TODO: We can actually do some scalar loads into the first and final vectors to get around alignment issues and then use the aligned load for all others.
-		while (Loops > 1)
-		{
-			FPlatformMath::WideVectorLoadHalf(DestReg, InputReg);
-			DestReg += 8;
-			InputReg += 8;
-			Loops -= 2;
-		}
-
-		while (Loops > 0)
-		{
-			FPlatformMath::VectorLoadHalf(DestReg, InputReg);
-			DestReg += 4;
-			InputReg += 4;
-			Loops -= 1;
-		}
-	}
-};
-
-
-/** Special kernel for reading from an input dataset; non-advancing (reads same instance everytime). 
- *  this kernel splats the X component of the source register to all 4 dest components; it's meant to
- *	use scalar data sets as the source (e.g. events)
- */
-template<typename T, int TypeOffset>
-struct FVectorKernelReadInputNoAdvance
-{
-	static void Optimize(FVectorVMCodeOptimizerContext& Context)
-	{
-		Context.WriteExecFunction(Exec);
-		Context.Write(Context.DecodeU16());	// DataSetIndex
-		Context.Write(Context.DecodeU16());	// InputRegisterIdx
-		Context.Write(Context.DecodeU16());	// DestRegisterIdx
-	}
-
-	static VM_FORCEINLINE void Exec(FVectorVMContext& Context)
-	{
-		static const int32 InstancesPerVector = sizeof(VectorRegister4Float) / sizeof(T);
-
-		const int32 DataSetIndex = Context.DecodeU16();
-		const int32 InputRegisterIdx = Context.DecodeU16();
-		const int32 DestRegisterIdx = Context.DecodeU16();
-		const int32 Loops = Context.GetNumLoops<InstancesPerVector>();
-
-		VectorRegister4Float* DestReg = (VectorRegister4Float*)(Context.GetTempRegister(DestRegisterIdx));
-		VectorRegister4Float* InputReg = (VectorRegister4Float*)(Context.GetInputRegister<T, TypeOffset>(DataSetIndex, InputRegisterIdx));
-
-		//TODO: We can actually do some scalar loads into the first and final vectors to get around alignment issues and then use the aligned load for all others.
-		for (int32 i = 0; i < Loops; ++i)
-		{
-			*DestReg = VectorSwizzle(VectorLoad(InputReg), 0,0,0,0);
-			++DestReg;
-		}
-	}
-};
-
-template<>
-struct FVectorKernelReadInputNoAdvance<FFloat16, 2>
-{
-	static void Optimize(FVectorVMCodeOptimizerContext& Context)
-	{
-		Context.WriteExecFunction(Exec);
-		Context.Write(Context.DecodeU16());	// DataSetIndex
-		Context.Write(Context.DecodeU16());	// InputRegisterIdx
-		Context.Write(Context.DecodeU16());	// DestRegisterIdx
-	}
-
-	static VM_FORCEINLINE void Exec(FVectorVMContext& Context)
-	{
-		const int32 DataSetIndex = Context.DecodeU16();
-		const int32 InputRegisterIdx = Context.DecodeU16();
-		const int32 DestRegisterIdx = Context.DecodeU16();
-		const int32 Loops = Context.GetNumLoops<4>();
-
-		VectorRegister4Float* DestReg = (VectorRegister4Float*)(Context.GetTempRegister(DestRegisterIdx));
-		uint16* InputReg = (uint16*)(Context.GetInputRegister<FFloat16, 2>(DataSetIndex, InputRegisterIdx));
-
-		//TODO: We can actually do some scalar loads into the first and final vectors to get around alignment issues and then use the aligned load for all others.
-		for (int32 i = 0; i < Loops; ++i)
-		{
-			float flt = FPlatformMath::LoadHalf(InputReg);
-			*DestReg = VectorLoadFloat1(&flt);
-			++DestReg;
-		}
-	}
-};
-
-//TODO - Should be straight forwards to follow the input with a mix of the outputs direct indexing
-/** Special kernel for reading an specific location in an input register. */
-// template<typename T>
-// struct FScalarKernelReadInputIndexed
-// {
-// 	static VM_FORCEINLINE void Exec(FVectorVMContext& Context)
-// 	{
-// 		int32* IndexReg = (int32*)(Context.RegisterTable[DecodeU16(Context)]);
-// 		T* InputReg = (T*)(Context.RegisterTable[DecodeU16(Context)]);
-// 		T* DestReg = (T*)(Context.RegisterTable[DecodeU16(Context)]);
-// 
-// 		//Has to be scalar as each instance can read from a different location in the input buffer.
-// 		for (int32 i = 0; i < Context.NumInstances; ++i)
-// 		{
-// 			T* ReadPtr = (*InputReg) + (*IndexReg);
-// 			*DestReg = (*ReadPtr);
-// 			++IndexReg;
-// 			++DestReg;
-// 		}
-// 	}
-// };
-
-/** Special kernel for writing to a specific output register. */
-template<typename SourceType, typename DestType, int TypeOffset>
-struct FScalarKernelWriteOutputIndexed
-{
-	static VM_FORCEINLINE void Optimize(FVectorVMCodeOptimizerContext& Context)
-	{
-		const uint32 SrcOpTypes = Context.BaseContext.DecodeSrcOperandTypes();
-		switch (SrcOpTypes)
-		{
-			case SRCOP_RRR: Context.WriteExecFunction(DoKernel<FRegisterHandler<SourceType>>); break;
-			case SRCOP_RRC:	Context.WriteExecFunction(DoKernel<FConstantHandler<SourceType>>); break;
-			default: check(0); break;
-		};
-
-		Context.Write(Context.DecodeU16());		// DataSetIndex
-		Context.Write(Context.DecodeU16());		// DestIndexRegisterIdx
-		Context.Write(Context.DecodeU16());		// DataHandlerType
-		Context.Write(Context.DecodeU16());		// DestRegisterIdx
-	}
-
-	static VM_FORCEINLINE void Exec(FVectorVMContext& Context)
-	{
-		const uint32 SrcOpTypes = Context.DecodeSrcOperandTypes();
-		switch (SrcOpTypes)
-		{
-		case SRCOP_RRR: DoKernel<FRegisterHandler<SourceType>>(Context); break;
-		case SRCOP_RRC:	DoKernel<FConstantHandler<SourceType>>(Context); break;
-		default: check(0); break;
-		};
-	}
-
-	template<typename DataHandlerType>
-	static VM_FORCEINLINE void DoKernel(FVectorVMContext& Context)
-	{
-		const int32 DataSetIndex = Context.DecodeU16();
-
-		const int32 DestIndexRegisterIdx = Context.DecodeU16();
-		int32* RESTRICT DestIndexReg = (int32*)(Context.GetTempRegister(DestIndexRegisterIdx));
-
-		DataHandlerType DataHandler(Context);
-
-		const int32 DestRegisterIdx = Context.DecodeU16();
-		DestType* RESTRICT DestReg = Context.GetOutputRegister<DestType, TypeOffset>(DataSetIndex, DestRegisterIdx);
-
-		int NumInstances = Context.GetNumInstances();
-		for (int32 i = 0; i < NumInstances; ++i)
-		{
-			int32 DestIndex = *DestIndexReg;
-			if (DestIndex != INDEX_NONE)
-			{
-				DestReg[DestIndex] = DataHandler.Get();
-			}
-
-			++DestIndexReg;
-			DataHandler.Advance();
-			//We don't increment the dest as we index into it directly.
-		}
-	}
-};
-
-template<>
-struct FScalarKernelWriteOutputIndexed<float, FFloat16, 2>
-{
-	static VM_FORCEINLINE void Optimize(FVectorVMCodeOptimizerContext& Context)
-	{
-		const uint32 SrcOpTypes = Context.BaseContext.DecodeSrcOperandTypes();
-		switch (SrcOpTypes)
-		{
-		case SRCOP_RRR: Context.WriteExecFunction(DoKernel<FRegisterHandler<float>>); break;
-		case SRCOP_RRC:	Context.WriteExecFunction(DoKernel<FConstantHandler<float>>); break;
-		default: check(0); break;
-		};
-
-		Context.Write(Context.DecodeU16());		// DataSetIndex
-		Context.Write(Context.DecodeU16());		// DestIndexRegisterIdx
-		Context.Write(Context.DecodeU16());		// DataHandlerType
-		Context.Write(Context.DecodeU16());		// DestRegisterIdx
-	}
-
-	static VM_FORCEINLINE void Exec(FVectorVMContext& Context)
-	{
-		const uint32 SrcOpTypes = Context.DecodeSrcOperandTypes();
-		switch (SrcOpTypes)
-		{
-		case SRCOP_RRR: DoKernel<FRegisterHandler<float>>(Context); break;
-		case SRCOP_RRC:	DoKernel<FConstantHandler<float>>(Context); break;
-		default: check(0); break;
-		};
-	}
-
-	template<typename DataHandlerType>
-	static VM_FORCEINLINE void DoKernel(FVectorVMContext& Context)
-	{
-		const int32 DataSetIndex = Context.DecodeU16();
-
-		const int32 DestIndexRegisterIdx = Context.DecodeU16();
-		int32* RESTRICT DestIndexReg = (int32*)(Context.GetTempRegister(DestIndexRegisterIdx));
-
-		DataHandlerType DataHandler(Context);
-
-		const int32 DestRegisterIdx = Context.DecodeU16();
-		uint16* RESTRICT DestReg = (uint16*)Context.GetOutputRegister<FFloat16, 2>(DataSetIndex, DestRegisterIdx);
-
-		int NumInstances = Context.GetNumInstances();
-		for (int32 i = 0; i < NumInstances; ++i)
-		{
-			int32 DestIndex = *DestIndexReg;
-			if (DestIndex != INDEX_NONE)
-			{
-				FPlatformMath::StoreHalf(&DestReg[DestIndex], DataHandler.Get());
-			}
-
-			++DestIndexReg;
-			DataHandler.Advance();
-			//We don't increment the dest as we index into it directly.
-		}
-	}
-};
-
-struct FDataSetCounterHandler
-{
-	int32* Counter;
-	FDataSetCounterHandler(FVectorVMContext& Context)
-		: Counter(&Context.GetDataSetMeta(Context.DecodeU16()).DataSetAccessIndex)
-	{}
-
-	VM_FORCEINLINE void Advance() { }
-	VM_FORCEINLINE int32* Get() { return Counter; }
-	VM_FORCEINLINE int32* GetAndAdvance() { return Counter; }
-	//VM_FORCEINLINE const int32* GetDest() { return Counter; }Should never use as a dest. All kernels with read and write to this.
-
-	static void Optimize(FVectorVMCodeOptimizerContext& Context)
-	{
-		Context.Write(Context.DecodeU16());
-	}
-};
-
-struct FScalarKernelAcquireCounterIndex
-{
-	template<bool bThreadsafe>
-	struct InternalKernel
-	{
-		static VM_FORCEINLINE void DoKernel(FVectorVMContext& Context, int32* RESTRICT Dst, int32* Index, int32 Valid)
-		{
-			if (Valid != 0)
-			{
-				*Dst = bThreadsafe ? FPlatformAtomics::InterlockedIncrement(Index) : ++(*Index);
-			}
-			else
-			{
-				*Dst = INDEX_NONE;	// Subsequent DoKernal calls above will skip over INDEX_NONE register entries...
-			}
-		}
-
-		static VM_FORCEINLINE void Exec(FVectorVMContext& Context)
-		{
-			const uint32 SrcOpType = Context.DecodeSrcOperandTypes();
-			switch (SrcOpType)
-			{
-				case SRCOP_RRR: TBinaryKernelHandler<InternalKernel<true>, FRegisterHandler<int32>, FDataSetCounterHandler, FRegisterHandler<int32>, 1>::Exec(Context); break;
-				case SRCOP_RRC:	TBinaryKernelHandler<InternalKernel<true>, FRegisterHandler<int32>, FDataSetCounterHandler, FConstantHandler<int32>, 1>::Exec(Context); break;
-				default: check(0); break;
-			};
-		}
-	};
-
-	template<uint32 SrcOpType>
-	static void ExecOptimized(FVectorVMContext& Context)
-	{
-		if (Context.IsParallelExecution())
-		{
-			switch (SrcOpType)
-			{
-				case SRCOP_RRR: TBinaryKernelHandler<InternalKernel<true>, FRegisterHandler<int32>, FDataSetCounterHandler, FRegisterHandler<int32>, 1>::Exec(Context); break;
-				case SRCOP_RRC:	TBinaryKernelHandler<InternalKernel<true>, FRegisterHandler<int32>, FDataSetCounterHandler, FConstantHandler<int32>, 1>::Exec(Context); break;
-				default: check(0); break;
-			}
-		}
-		else
-		{
-			switch (SrcOpType)
-			{
-				case SRCOP_RRR: TBinaryKernelHandler<InternalKernel<false>, FRegisterHandler<int32>, FDataSetCounterHandler, FRegisterHandler<int32>, 1>::Exec(Context); break;
-				case SRCOP_RRC:	TBinaryKernelHandler<InternalKernel<false>, FRegisterHandler<int32>, FDataSetCounterHandler, FConstantHandler<int32>, 1>::Exec(Context); break;
-				default: check(0); break;
-			}
-		}
-	}
-
-	static void Optimize(FVectorVMCodeOptimizerContext& Context)
-	{
-		const uint32 SrcOpType = Context.BaseContext.DecodeSrcOperandTypes();
-		switch (SrcOpType)
-		{
-			case SRCOP_RRR: Context.WriteExecFunction(FScalarKernelAcquireCounterIndex::ExecOptimized<SRCOP_RRR>); break;
-			case SRCOP_RRC: Context.WriteExecFunction(FScalarKernelAcquireCounterIndex::ExecOptimized<SRCOP_RRC>); break;
-			default: check(0); break;
-		}
-
-		// Three registers, note we don't call Optimize on the Kernel since that will write the Exec and we are selecting based upon thread safe or not
-		Context.Write(Context.DecodeU16());
-		Context.Write(Context.DecodeU16());
-		Context.Write(Context.DecodeU16());
-	}
-
-	static VM_FORCEINLINE void Exec(FVectorVMContext& Context)
-	{
-		if ( Context.IsParallelExecution() )
-		{
-			InternalKernel<true>::Exec(Context);
-		}
-		else
-		{
-			InternalKernel<false>::Exec(Context);
-		}
-	}
-
-};
-
-//TODO: REWORK TO FUNCITON LIKE THE ABOVE.
-// /** Special kernel for decrementing a dataset counter. */
-// struct FScalarKernelReleaseCounterIndex
-// {
-// 	static VM_FORCEINLINE void Exec(FVectorVMContext& Context)
-// 	{
-// 		int32* CounterPtr = (int32*)(Context.ConstantTable[DecodeU16(Context)]);
-// 		int32* DestReg = (int32*)(Context.RegisterTable[DecodeU16(Context)]);
-// 
-// 		for (int32 i = 0; i < Context.NumInstances; ++i)
-// 		{
-// 			int32 Counter = (*CounterPtr--);
-// 			*DestReg = Counter >= 0 ? Counter : INDEX_NONE;
-// 
-// 			++DestReg;
-// 		}
-// 	}
-// };
-
-//////////////////////////////////////////////////////////////////////////
-//external_func_call
-
-struct FKernelExternalFunctionCall
-{
-	static void Optimize(FVectorVMCodeOptimizerContext& Context)
-	{
-		const uint32 ExternalFuncIdx = Context.DecodeU8();
-
-		Context.WriteExecFunction(Exec);
-		Context.Write<uint8>(ExternalFuncIdx);
-
-		const int32 NumRegisters = Context.ExternalFunctionRegisterCounts[ExternalFuncIdx];
-		for ( int32 i=0; i < NumRegisters; ++i )
-		{
-			Context.Write(Context.DecodeU16());
-		}
-	}
-
-	static void Exec(FVectorVMContext& Context)
-	{
-#if VECTORVM_SUPPORTS_LEGACY
-		const uint32 ExternalFuncIdx = Context.DecodeU8();
-		const FVMExternalFunction* ExternalFunction = Context.ExternalFunctionTable[ExternalFuncIdx];
-		check(ExternalFunction);
-
-		if (ExternalFunction)
-		{
-#if VECTORVM_SUPPORTS_EXPERIMENTAL
-			FVectorVMExternalFunctionContext DataInterfaceFunctionContext(&Context);
-#else
-			FVectorVMExternalFunctionContextLegacy DataInterfaceFunctionContext(&Context);
-#endif
-			ExternalFunction->Execute(DataInterfaceFunctionContext);
-		}
-#endif
-	}
-};
-
-//////////////////////////////////////////////////////////////////////////
-//Integer operations
-
-//addi,
-struct FVectorIntKernelAdd : TBinaryVectorIntKernel<FVectorIntKernelAdd>
-{
-	static void VM_FORCEINLINE DoKernel(FVectorVMContext& Context, VectorRegister4Int* Dst, VectorRegister4Int Src0, VectorRegister4Int Src1)
-	{
-		*Dst = VectorIntAdd(Src0, Src1);
-	}
-};
-
-//subi,
-struct FVectorIntKernelSubtract : TBinaryVectorIntKernel<FVectorIntKernelSubtract>
-{
-	static void VM_FORCEINLINE DoKernel(FVectorVMContext& Context, VectorRegister4Int* Dst, VectorRegister4Int Src0, VectorRegister4Int Src1)
-	{
-		*Dst = VectorIntSubtract(Src0, Src1);
-	}
-};
-
-//muli,
-struct FVectorIntKernelMultiply : TBinaryVectorIntKernel<FVectorIntKernelMultiply>
-{
-	static void VM_FORCEINLINE DoKernel(FVectorVMContext& Context, VectorRegister4Int* Dst, VectorRegister4Int Src0, VectorRegister4Int Src1)
-	{
-		*Dst = VectorIntMultiply(Src0, Src1);
-	}
-};
-
-//divi,
-struct FVectorIntKernelDivide : TBinaryVectorIntKernel<FVectorIntKernelDivide>
-{
-	static void VM_FORCEINLINE DoKernel(FVectorVMContext& Context, VectorRegister4Int* Dst, VectorRegister4Int Src0, VectorRegister4Int Src1)
-	{
-		int32 TmpA[4];
-		VectorIntStore(Src0, TmpA);
-
-		int32 TmpB[4];
-		VectorIntStore(Src1, TmpB);
-
-		// No intrinsics exist for integer divide. Since div by zero causes crashes, we must be safe against that.
-		int32 TmpDst[4];
-		TmpDst[0] = SafeIntDivide(TmpA[0], TmpB[0]);
-		TmpDst[1] = SafeIntDivide(TmpA[1], TmpB[1]);
-		TmpDst[2] = SafeIntDivide(TmpA[2], TmpB[2]);
-		TmpDst[3] = SafeIntDivide(TmpA[3], TmpB[3]);
-
-		*Dst = MakeVectorRegisterInt(TmpDst[0], TmpDst[1], TmpDst[2], TmpDst[3]);
-	}
-
-private:
-	VM_FORCEINLINE static int32 SafeIntDivide(int32 Numerator, int32 Denominator)
-	{
-		static constexpr int32 MinIntValue = std::numeric_limits<int32>::min();
-		static constexpr int32 MaxIntValue = std::numeric_limits<int32>::max();
-
-		if (Denominator == 0)
-		{
-			return 0;
-		}
-		else if ((Denominator == -1) && (Numerator == MinIntValue))
-		{
-			return MaxIntValue;
-		}
-
-		return Numerator / Denominator;
-	}
-};
-
-
-//clampi,
-struct FVectorIntKernelClamp : TTrinaryVectorIntKernel<FVectorIntKernelClamp>
-{
-	static void VM_FORCEINLINE DoKernel(FVectorVMContext& Context, VectorRegister4Int* Dst, VectorRegister4Int Src0, VectorRegister4Int Src1, VectorRegister4Int Src2)
-	{
-		*Dst = VectorIntMin(VectorIntMax(Src0, Src1), Src2);
-	}
-};
-
-//mini,
-struct FVectorIntKernelMin : TBinaryVectorIntKernel<FVectorIntKernelMin>
-{
-	static void VM_FORCEINLINE DoKernel(FVectorVMContext& Context, VectorRegister4Int* Dst, VectorRegister4Int Src0, VectorRegister4Int Src1)
-	{
-		*Dst = VectorIntMin(Src0, Src1);
-	}
-};
-
-//maxi,
-struct FVectorIntKernelMax : TBinaryVectorIntKernel<FVectorIntKernelMax>
-{
-	static void VM_FORCEINLINE DoKernel(FVectorVMContext& Context, VectorRegister4Int* Dst, VectorRegister4Int Src0, VectorRegister4Int Src1)
-	{
-		*Dst = VectorIntMax(Src0, Src1);
-	}
-};
-
-//absi,
-struct FVectorIntKernelAbs : TUnaryVectorIntKernel<FVectorIntKernelAbs>
-{
-	static void VM_FORCEINLINE DoKernel(FVectorVMContext& Context, VectorRegister4Int* Dst, VectorRegister4Int Src0)
-	{
-		*Dst = VectorIntAbs(Src0);
-	}
-};
-
-//negi,
-struct FVectorIntKernelNegate : TUnaryVectorIntKernel<FVectorIntKernelNegate>
-{
-	static void VM_FORCEINLINE DoKernel(FVectorVMContext& Context, VectorRegister4Int* Dst, VectorRegister4Int Src0)
-	{
-		*Dst = VectorIntNegate(Src0);
-	}
-};
-
-//signi,
-struct FVectorIntKernelSign : TUnaryVectorIntKernel<FVectorIntKernelSign>
-{
-	static void VM_FORCEINLINE DoKernel(FVectorVMContext& Context, VectorRegister4Int* Dst, VectorRegister4Int Src0)
-	{
-		*Dst = VectorIntSign(Src0);
-	}
-};
-
-//randomi,
-//No good way to do this with SSE atm so just do it scalar.
-struct FScalarIntKernelRandom : public TUnaryScalarIntKernel<FScalarIntKernelRandom>
-{
-	static void VM_FORCEINLINE DoKernel(FVectorVMContext& Context, int32* RESTRICT Dst, int32 Src0)
-	{
-		//EEK!. Improve this. Implement GPU style seeded rand instead of this.
-		*Dst = static_cast<int32>(Context.RandStream.GetFraction() * Src0);
-	}
-};
-
-//cmplti,
-struct FVectorIntKernelCompareLT : TBinaryVectorIntKernel<FVectorIntKernelCompareLT>
-{
-	static void VM_FORCEINLINE DoKernel(FVectorVMContext& Context, VectorRegister4Int* Dst, VectorRegister4Int Src0, VectorRegister4Int Src1)
-	{
-		*Dst = VectorIntCompareLT(Src0, Src1);
-	}
-};
-
-//cmplei,
-struct FVectorIntKernelCompareLE : TBinaryVectorIntKernel<FVectorIntKernelCompareLE>
-{
-	static void VM_FORCEINLINE DoKernel(FVectorVMContext& Context, VectorRegister4Int* Dst, VectorRegister4Int Src0, VectorRegister4Int Src1)
-	{
-		*Dst = VectorIntCompareLE(Src0, Src1);
-	}
-};
-
-//cmpgti,
-struct FVectorIntKernelCompareGT : TBinaryVectorIntKernel<FVectorIntKernelCompareGT>
-{
-	static void VM_FORCEINLINE DoKernel(FVectorVMContext& Context, VectorRegister4Int* Dst, VectorRegister4Int Src0, VectorRegister4Int Src1)
-	{
-		*Dst = VectorIntCompareGT(Src0, Src1);
-	}
-};
-
-//cmpgei,
-struct FVectorIntKernelCompareGE : TBinaryVectorIntKernel<FVectorIntKernelCompareGE>
-{
-	static void VM_FORCEINLINE DoKernel(FVectorVMContext& Context, VectorRegister4Int* Dst, VectorRegister4Int Src0, VectorRegister4Int Src1)
-	{
-		*Dst = VectorIntCompareGE(Src0, Src1);
-	}
-};
-
-//cmpeqi,
-struct FVectorIntKernelCompareEQ : TBinaryVectorIntKernel<FVectorIntKernelCompareEQ>
-{
-	static void VM_FORCEINLINE DoKernel(FVectorVMContext& Context, VectorRegister4Int* Dst, VectorRegister4Int Src0, VectorRegister4Int Src1)
-	{
-		*Dst = VectorIntCompareEQ(Src0, Src1);
-	}
-};
-
-//cmpneqi,
-struct FVectorIntKernelCompareNEQ : TBinaryVectorIntKernel<FVectorIntKernelCompareNEQ>
-{
-	static void VM_FORCEINLINE DoKernel(FVectorVMContext& Context, VectorRegister4Int* Dst, VectorRegister4Int Src0, VectorRegister4Int Src1)
-	{
-		*Dst = VectorIntCompareNEQ(Src0, Src1);
-	}
-};
-
-//bit_and,
-struct FVectorIntKernelBitAnd : TBinaryVectorIntKernel<FVectorIntKernelBitAnd>
-{
-	static void VM_FORCEINLINE DoKernel(FVectorVMContext& Context, VectorRegister4Int* Dst, VectorRegister4Int Src0, VectorRegister4Int Src1)
-	{
-		*Dst = VectorIntAnd(Src0, Src1);
-	}
-};
-
-//bit_or,
-struct FVectorIntKernelBitOr : TBinaryVectorIntKernel<FVectorIntKernelBitOr>
-{
-	static void VM_FORCEINLINE DoKernel(FVectorVMContext& Context, VectorRegister4Int* Dst, VectorRegister4Int Src0, VectorRegister4Int Src1)
-	{
-		*Dst = VectorIntOr(Src0, Src1);
-	}
-};
-
-//bit_xor,
-struct FVectorIntKernelBitXor : TBinaryVectorIntKernel<FVectorIntKernelBitXor>
-{
-	static void VM_FORCEINLINE DoKernel(FVectorVMContext& Context, VectorRegister4Int* Dst, VectorRegister4Int Src0, VectorRegister4Int Src1)
-	{
-		*Dst = VectorIntXor(Src0, Src1);
-	}
-};
-
-//bit_not,
-struct FVectorIntKernelBitNot : TUnaryVectorIntKernel<FVectorIntKernelBitNot>
-{
-	static void VM_FORCEINLINE DoKernel(FVectorVMContext& Context, VectorRegister4Int* Dst, VectorRegister4Int Src0)
-	{
-		*Dst = VectorIntNot(Src0);
-	}
-};
-
-// bit_lshift
-struct FVectorIntKernelBitLShift : TBinaryVectorIntKernel<FVectorIntKernelBitLShift>
-{
-	static void VM_FORCEINLINE DoKernel(FVectorVMContext& Context, VectorRegister4Int* Dst, VectorRegister4Int Src0,  VectorRegister4Int Src1)
-	{
-		int32 TmpA[4];
-		VectorIntStore(Src0, TmpA);
-
-		int32 TmpB[4];
-		VectorIntStore(Src1, TmpB);
-
-		int32 TmpDst[4];
-		TmpDst[0] = (TmpA[0] << TmpB[0]);
-		TmpDst[1] = (TmpA[1] << TmpB[1]);
-		TmpDst[2] = (TmpA[2] << TmpB[2]);
-		TmpDst[3] = (TmpA[3] << TmpB[3]);
-		*Dst = MakeVectorRegisterInt(TmpDst[0], TmpDst[1], TmpDst[2], TmpDst[3] );
-	}
-};
-
-// bit_rshift
-struct FVectorIntKernelBitRShift : TBinaryVectorIntKernel<FVectorIntKernelBitRShift>
-{
-	static void VM_FORCEINLINE DoKernel(FVectorVMContext& Context, VectorRegister4Int* Dst, VectorRegister4Int Src0, VectorRegister4Int Src1)
-	{
-		int32 TmpA[4];
-		VectorIntStore(Src0, TmpA);
-
-		int32 TmpB[4];
-		VectorIntStore(Src1, TmpB);
-
-		int32 TmpDst[4];
-		TmpDst[0] = (TmpA[0] >> TmpB[0]);
-		TmpDst[1] = (TmpA[1] >> TmpB[1]);
-		TmpDst[2] = (TmpA[2] >> TmpB[2]);
-		TmpDst[3] = (TmpA[3] >> TmpB[3]);
-		*Dst = MakeVectorRegisterInt(TmpDst[0], TmpDst[1], TmpDst[2], TmpDst[3]);
-	}
-};
-
-//"Boolean" ops. Currently handling bools as integers.
-//logic_and,
-struct FVectorIntKernelLogicAnd : TBinaryVectorIntKernel<FVectorIntKernelLogicAnd>
-{
-	static void VM_FORCEINLINE DoKernel(FVectorVMContext& Context, VectorRegister4Int* Dst, VectorRegister4Int Src0, VectorRegister4Int Src1)
-	{
-		//We need to assume a mask input and produce a mask output so just bitwise ops actually fine for these?
-		*Dst = VectorIntAnd(Src0, Src1);
-	}
-};
-
-//logic_or,
-struct FVectorIntKernelLogicOr : TBinaryVectorIntKernel<FVectorIntKernelLogicOr>
-{
-	static void VM_FORCEINLINE DoKernel(FVectorVMContext& Context, VectorRegister4Int* Dst, VectorRegister4Int Src0, VectorRegister4Int Src1)
-	{
-		//We need to assume a mask input and produce a mask output so just bitwise ops actually fine for these?
-		*Dst = VectorIntOr(Src0, Src1);
-	}
-};
-//logic_xor,
-struct FVectorIntKernelLogicXor : TBinaryVectorIntKernel<FVectorIntKernelLogicXor>
-{
-	static void VM_FORCEINLINE DoKernel(FVectorVMContext& Context, VectorRegister4Int* Dst, VectorRegister4Int Src0, VectorRegister4Int Src1)
-	{
-		//We need to assume a mask input and produce a mask output so just bitwise ops actually fine for these?
-		*Dst = VectorIntXor(Src0, Src1);
-	}
-};
-
-//logic_not,
-struct FVectorIntKernelLogicNot : TUnaryVectorIntKernel<FVectorIntKernelLogicNot>
-{
-	static void VM_FORCEINLINE DoKernel(FVectorVMContext& Context, VectorRegister4Int* Dst, VectorRegister4Int Src0)
-	{
-		//We need to assume a mask input and produce a mask output so just bitwise ops actually fine for these?
-		*Dst = VectorIntNot(Src0);
-	}
-};
-
-//conversions
-//f2i,
-struct FVectorKernelFloatToInt : TUnaryKernel<FVectorKernelFloatToInt, FRegisterHandler<VectorRegister4Int>, FConstantHandler<VectorRegister4Float>, FRegisterHandler<VectorRegister4Float>, VECTOR_WIDTH_FLOATS>
-{
-	static void VM_FORCEINLINE DoKernel(FVectorVMContext& Context, VectorRegister4Int* Dst, VectorRegister4Float Src0)
-	{
-		*Dst = VectorFloatToInt(Src0);
-	}
-};
-
-//i2f,
-struct FVectorKernelIntToFloat : TUnaryKernel<FVectorKernelIntToFloat, FRegisterHandler<VectorRegister4Float>, FConstantHandler<VectorRegister4Int>, FRegisterHandler<VectorRegister4Int>, VECTOR_WIDTH_FLOATS>
-{
-	static void VM_FORCEINLINE DoKernel(FVectorVMContext& Context, VectorRegister4Float* Dst, VectorRegister4Int Src0)
-	{
-		*Dst = VectorIntToFloat(Src0);
-	}
-};
-
-//f2b,
-struct FVectorKernelFloatToBool : TUnaryKernel<FVectorKernelFloatToBool, FRegisterHandler<VectorRegister4Float>, FConstantHandler<VectorRegister4Float>, FRegisterHandler<VectorRegister4Float>, VECTOR_WIDTH_FLOATS>
-{
-	static void VM_FORCEINLINE DoKernel(FVectorVMContext& Context, VectorRegister4Float* Dst, VectorRegister4Float Src0)
-	{		
-		*Dst = VectorCompareGT(Src0, GlobalVectorConstants::FloatZero);
-	}
-};
-
-//b2f,
-struct FVectorKernelBoolToFloat : TUnaryKernel<FVectorKernelBoolToFloat, FRegisterHandler<VectorRegister4Float>, FConstantHandler<VectorRegister4Float>, FRegisterHandler<VectorRegister4Float>, VECTOR_WIDTH_FLOATS>
-{
-	static void VM_FORCEINLINE DoKernel(FVectorVMContext& Context, VectorRegister4Float* Dst, VectorRegister4Float Src0)
-	{
-		*Dst = VectorSelect(Src0, GlobalVectorConstants::FloatOne, GlobalVectorConstants::FloatZero);
-	}
-};
-
-//i2b,
-struct FVectorKernelIntToBool : TUnaryKernel<FVectorKernelIntToBool, FRegisterHandler<VectorRegister4Int>, FConstantHandler<VectorRegister4Int>, FRegisterHandler<VectorRegister4Int>, VECTOR_WIDTH_FLOATS>
-{
-	static void VM_FORCEINLINE DoKernel(FVectorVMContext& Context, VectorRegister4Int* Dst, VectorRegister4Int Src0)
-	{
-		*Dst = VectorIntCompareGT(Src0, GlobalVectorConstants::IntZero);
-	}
-};
-
-//b2i,
-struct FVectorKernelBoolToInt : TUnaryKernel<FVectorKernelBoolToInt, FRegisterHandler<VectorRegister4Int>, FConstantHandler<VectorRegister4Int>, FRegisterHandler<VectorRegister4Int>, VECTOR_WIDTH_FLOATS>
-{
-	static void VM_FORCEINLINE DoKernel(FVectorVMContext& Context, VectorRegister4Int* Dst, VectorRegister4Int Src0)
-	{
-		*Dst = VectorIntSelect(Src0, GlobalVectorConstants::IntOne, GlobalVectorConstants::IntZero);
-	}
-};
-
-//reinterpret bits
-//fasi,
-struct FVectorKernelFloatAsInt : TUnaryKernel<FVectorKernelFloatAsInt, FRegisterHandler<VectorRegister4Int>, FConstantHandler<VectorRegister4Float>, FRegisterHandler<VectorRegister4Float>, VECTOR_WIDTH_FLOATS>
-{
-	static void VM_FORCEINLINE DoKernel(FVectorVMContext& Context, VectorRegister4Int* Dst, VectorRegister4Float Src0)
-	{
-		*Dst = VectorCastFloatToInt(Src0);
-	}
-};
-
-//iasf,
-struct FVectorKernelIntAsFloat : TUnaryKernel<FVectorKernelIntAsFloat, FRegisterHandler<VectorRegister4Float>, FConstantHandler<VectorRegister4Int>, FRegisterHandler<VectorRegister4Int>, VECTOR_WIDTH_FLOATS>
-{
-	static void VM_FORCEINLINE DoKernel(FVectorVMContext& Context, VectorRegister4Float* Dst, VectorRegister4Int Src0)
-	{
-		*Dst = VectorCastIntToFloat(Src0);
-	}
-};
-
-void VectorVM::Exec(FVectorVMExecArgs& Args, FVectorVMSerializeState *SerializeState)
-{
-	//TRACE_CPUPROFILER_EVENT_SCOPE(VMExec);
-	SCOPE_CYCLE_COUNTER(STAT_VVMExec);
-
-#if UE_BUILD_TEST
-	const bool bNumInstancesEvent = GbDetailedVMScriptStats != 0;
-	if (bNumInstancesEvent)
-	{
-		FPlatformMisc::BeginNamedEvent(FColor::Red, *FString::Printf(TEXT("STAT_VVMExec - %d"), Args.NumInstances));
-	}
-#endif
-
-	const int32 MaxInstances = FMath::Min(GParallelVVMInstancesPerChunk, Args.NumInstances);
-	const int32 NumChunks = (Args.NumInstances / GParallelVVMInstancesPerChunk) + 1;
-	const int32 ChunksPerBatch = (GbParallelVVM != 0 && FApp::ShouldUseThreadingForPerformance()) ? GParallelVVMChunksPerBatch : NumChunks;
-	const int32 NumBatches = FMath::DivideAndRoundUp(NumChunks, ChunksPerBatch);
-	const bool bParallel = NumBatches > 1 && Args.bAllowParallel;
-#	ifdef VVM_INCLUDE_SERIALIZATION
-	const bool bUseOptimizedByteCode = false; //serializes the bytecode from the instructions, cannot use jump table
-#	else //VVM_INCLUDE_SERIALIZATION
-	const bool bUseOptimizedByteCode = (Args.OptimizedByteCode != nullptr) && GbUseOptimizedVMByteCode;
-#	endif
-	const FVectorVMExecFunction* OptimizedJumpTable = bUseOptimizedByteCode ? FVectorVMCodeOptimizerContext::DecodeJumpTable(Args.OptimizedByteCode) : nullptr;
-	
-	uint64 StartTime = FPlatformTime::Cycles64();
-
-	auto ExecChunkBatch = [&](int32 BatchIdx)
-	{
-		//SCOPE_CYCLE_COUNTER(STAT_VVMExecChunk);
-
-		FVectorVMContext& Context = FVectorVMContext::Get();
-		Context.PrepareForExec(Args.NumTempRegisters, Args.ConstantTableCount, Args.ConstantTable, Args.ConstantTableSizes, Args.ExternalFunctionTable, Args.UserPtrTable, Args.DataSetMetaTable, MaxInstances, bParallel);
-#if STATS
-		Context.SetStatScopes(Args.StatScopes);
-#elif ENABLE_STATNAMEDEVENTS
-		Context.SetStatNamedEventScopes(Args.StatNamedEventsScopes);
-#endif
-
-		// Process one chunk at a time.
-		int32 ChunkIdx = BatchIdx * ChunksPerBatch;
-		const int32 FirstInstance = ChunkIdx * GParallelVVMInstancesPerChunk;
-		const int32 FinalInstance = FMath::Min(Args.NumInstances, FirstInstance + (ChunksPerBatch * GParallelVVMInstancesPerChunk));
-		int32 InstancesLeft = FinalInstance - FirstInstance;
-		while (InstancesLeft > 0)
-		{
-			int32 NumInstancesThisChunk = FMath::Min(InstancesLeft, (int32)GParallelVVMInstancesPerChunk);
-			int32 StartInstance = GParallelVVMInstancesPerChunk * ChunkIdx;
-
-			// Execute optimized byte code version
-			if ( bUseOptimizedByteCode )
-			{
-				// Setup execution context.
-				Context.PrepareForChunk(Args.OptimizedByteCode, NumInstancesThisChunk, StartInstance);
-
-				while (true)
-				{
-					FVectorVMExecFunction ExecFunction = OptimizedJumpTable[Context.DecodeU8()];
-					if (ExecFunction == nullptr)
-					{
-						break;
-					}
-					ExecFunction(Context);
-				}
-			}
-			else
-			{
-				// Setup execution context.
-				Context.PrepareForChunk(Args.ByteCode, NumInstancesThisChunk, StartInstance);
-
-				VVMSer_chunkStart(Context, ChunkIdx, BatchIdx);
-				// Execute VM on all vectors in this chunk.
-				EVectorVMOp Op = EVectorVMOp::done;
-				do
-				{
-					VVMSer_insStart(Context);
-					Op = Context.DecodeOp();
-					switch (Op)
-					{
-						// Dispatch kernel ops.
-						case EVectorVMOp::add: FVectorKernelAdd::Exec(Context); break;
-						case EVectorVMOp::sub: FVectorKernelSub::Exec(Context); break;
-						case EVectorVMOp::mul: FVectorKernelMul::Exec(Context); break;
-						case EVectorVMOp::div: FVectorKernelDivSafe::Exec(Context); break;
-						case EVectorVMOp::mad: FVectorKernelMad::Exec(Context); break;
-						case EVectorVMOp::lerp: FVectorKernelLerp::Exec(Context); break;
-						case EVectorVMOp::rcp: FVectorKernelRcpSafe::Exec(Context); break;
-						case EVectorVMOp::rsq: FVectorKernelRsqSafe::Exec(Context); break;
-						case EVectorVMOp::sqrt: FVectorKernelSqrtSafe::Exec(Context); break;
-						case EVectorVMOp::neg: FVectorKernelNeg::Exec(Context); break;
-						case EVectorVMOp::abs: FVectorKernelAbs::Exec(Context); break;
-						case EVectorVMOp::exp: FVectorKernelExp::Exec(Context); break;
-						case EVectorVMOp::exp2: FVectorKernelExp2::Exec(Context); break;
-						case EVectorVMOp::log: FVectorKernelLogSafe::Exec(Context); break;
-						case EVectorVMOp::log2: FVectorKernelLog2::Exec(Context); break;
-						case EVectorVMOp::sin: FVectorKernelSin::Exec(Context); break;
-						case EVectorVMOp::cos: FVectorKernelCos::Exec(Context); break;
-						case EVectorVMOp::tan: FVectorKernelTan::Exec(Context); break;
-						case EVectorVMOp::asin: FVectorKernelASin::Exec(Context); break;
-						case EVectorVMOp::acos: FVectorKernelACos::Exec(Context); break;
-						case EVectorVMOp::atan: FVectorKernelATan::Exec(Context); break;
-						case EVectorVMOp::atan2: FVectorKernelATan2::Exec(Context); break;
-						case EVectorVMOp::ceil: FVectorKernelCeil::Exec(Context); break;
-						case EVectorVMOp::floor: FVectorKernelFloor::Exec(Context); break;
-						case EVectorVMOp::round: FVectorKernelRound::Exec(Context); break;
-						case EVectorVMOp::fmod: FVectorKernelMod::Exec(Context); break;
-						case EVectorVMOp::frac: FVectorKernelFrac::Exec(Context); break;
-						case EVectorVMOp::trunc: FVectorKernelTrunc::Exec(Context); break;
-						case EVectorVMOp::clamp: FVectorKernelClamp::Exec(Context); break;
-						case EVectorVMOp::min: FVectorKernelMin::Exec(Context); break;
-						case EVectorVMOp::max: FVectorKernelMax::Exec(Context); break;
-						case EVectorVMOp::pow: FVectorKernelPowSafe::Exec(Context); break;
-						case EVectorVMOp::sign: FVectorKernelSign::Exec(Context); break;
-						case EVectorVMOp::step: FVectorKernelStep::Exec(Context); break;
-						case EVectorVMOp::random: FVectorKernelRandom::Exec(Context); break;
-						case EVectorVMOp::noise: VectorVMNoise::Noise1D(Context); break;
-						case EVectorVMOp::noise2D: VectorVMNoise::Noise2D(Context); break;
-						case EVectorVMOp::noise3D: VectorVMNoise::Noise3D(Context); break;
-
-						case EVectorVMOp::cmplt: FVectorKernelCompareLT::Exec(Context); break;
-						case EVectorVMOp::cmple: FVectorKernelCompareLE::Exec(Context); break;
-						case EVectorVMOp::cmpgt: FVectorKernelCompareGT::Exec(Context); break;
-						case EVectorVMOp::cmpge: FVectorKernelCompareGE::Exec(Context); break;
-						case EVectorVMOp::cmpeq: FVectorKernelCompareEQ::Exec(Context); break;
-						case EVectorVMOp::cmpneq: FVectorKernelCompareNEQ::Exec(Context); break;
-						case EVectorVMOp::select: FVectorKernelSelect::Exec(Context); break;
-
-						case EVectorVMOp::addi: FVectorIntKernelAdd::Exec(Context); break;
-						case EVectorVMOp::subi: FVectorIntKernelSubtract::Exec(Context); break;
-						case EVectorVMOp::muli: FVectorIntKernelMultiply::Exec(Context); break;
-						case EVectorVMOp::divi: FVectorIntKernelDivide::Exec(Context); break;
-						case EVectorVMOp::clampi: FVectorIntKernelClamp::Exec(Context); break;
-						case EVectorVMOp::mini: FVectorIntKernelMin::Exec(Context); break;
-						case EVectorVMOp::maxi: FVectorIntKernelMax::Exec(Context); break;
-						case EVectorVMOp::absi: FVectorIntKernelAbs::Exec(Context); break;
-						case EVectorVMOp::negi: FVectorIntKernelNegate::Exec(Context); break;
-						case EVectorVMOp::signi: FVectorIntKernelSign::Exec(Context); break;
-						case EVectorVMOp::randomi: FScalarIntKernelRandom::Exec(Context); break;
-						case EVectorVMOp::cmplti: FVectorIntKernelCompareLT::Exec(Context); break;
-						case EVectorVMOp::cmplei: FVectorIntKernelCompareLE::Exec(Context); break;
-						case EVectorVMOp::cmpgti: FVectorIntKernelCompareGT::Exec(Context); break;
-						case EVectorVMOp::cmpgei: FVectorIntKernelCompareGE::Exec(Context); break;
-						case EVectorVMOp::cmpeqi: FVectorIntKernelCompareEQ::Exec(Context); break;
-						case EVectorVMOp::cmpneqi: FVectorIntKernelCompareNEQ::Exec(Context); break;
-						case EVectorVMOp::bit_and: FVectorIntKernelBitAnd::Exec(Context); break;
-						case EVectorVMOp::bit_or: FVectorIntKernelBitOr::Exec(Context); break;
-						case EVectorVMOp::bit_xor: FVectorIntKernelBitXor::Exec(Context); break;
-						case EVectorVMOp::bit_not: FVectorIntKernelBitNot::Exec(Context); break;
-						case EVectorVMOp::bit_lshift: FVectorIntKernelBitLShift::Exec(Context); break;
-						case EVectorVMOp::bit_rshift: FVectorIntKernelBitRShift::Exec(Context); break;
-						case EVectorVMOp::logic_and: FVectorIntKernelLogicAnd::Exec(Context); break;
-						case EVectorVMOp::logic_or: FVectorIntKernelLogicOr::Exec(Context); break;
-						case EVectorVMOp::logic_xor: FVectorIntKernelLogicXor::Exec(Context); break;
-						case EVectorVMOp::logic_not: FVectorIntKernelLogicNot::Exec(Context); break;
-						case EVectorVMOp::f2i: FVectorKernelFloatToInt::Exec(Context); break;
-						case EVectorVMOp::i2f: FVectorKernelIntToFloat::Exec(Context); break;
-						case EVectorVMOp::f2b: FVectorKernelFloatToBool::Exec(Context); break;
-						case EVectorVMOp::b2f: FVectorKernelBoolToFloat::Exec(Context); break;
-						case EVectorVMOp::i2b: FVectorKernelIntToBool::Exec(Context); break;
-						case EVectorVMOp::b2i: FVectorKernelBoolToInt::Exec(Context); break;
-						case EVectorVMOp::fasi: FVectorKernelFloatAsInt::Exec(Context); break;
-						case EVectorVMOp::iasf: FVectorKernelIntAsFloat::Exec(Context); break;
-
-						case EVectorVMOp::outputdata_half:	FScalarKernelWriteOutputIndexed<float, FFloat16, 2>::Exec(Context);	break;
-						case EVectorVMOp::inputdata_half: FVectorKernelReadInput<FFloat16, 2>::Exec(Context); break;
-						case EVectorVMOp::outputdata_int32:	FScalarKernelWriteOutputIndexed<int32, int32, 1>::Exec(Context);	break;
-						case EVectorVMOp::inputdata_int32: FVectorKernelReadInput<int32, 1>::Exec(Context); break;
-						case EVectorVMOp::outputdata_float:	FScalarKernelWriteOutputIndexed<float, float, 0>::Exec(Context);	break;
-						case EVectorVMOp::inputdata_float: FVectorKernelReadInput<float, 0>::Exec(Context); break;
-						case EVectorVMOp::inputdata_noadvance_int32: FVectorKernelReadInputNoAdvance<int32, 1>::Exec(Context); break;
-						case EVectorVMOp::inputdata_noadvance_float: FVectorKernelReadInputNoAdvance<float, 0>::Exec(Context); break;
-						case EVectorVMOp::inputdata_noadvance_half: FVectorKernelReadInputNoAdvance<FFloat16, 2>::Exec(Context); break;
-						case EVectorVMOp::acquireindex:	FScalarKernelAcquireCounterIndex::Exec(Context); break;
-						case EVectorVMOp::external_func_call: FKernelExternalFunctionCall::Exec(Context); break;
-
-						case EVectorVMOp::exec_index: FVectorKernelExecutionIndex::Exec(Context); break;
-
-						case EVectorVMOp::enter_stat_scope: FVectorKernelEnterStatScope::Exec(Context); break;
-						case EVectorVMOp::exit_stat_scope: FVectorKernelExitStatScope::Exec(Context); break;
-
-						//Special case ops to handle unique IDs but this can be written as generalized buffer operations. TODO!
-						case EVectorVMOp::update_id:	FScalarKernelUpdateID::Exec(Context); break;
-						case EVectorVMOp::acquire_id:	FScalarKernelAcquireID::Exec(Context); break;
-
-						// Execution always terminates with a "done" opcode.
-						case EVectorVMOp::done:
-							break;
-
-						// Opcode not recognized / implemented.
-						default:
-							UE_LOG(LogVectorVM, Fatal, TEXT("Unknown op code 0x%02x"), (uint32)Op);
-							return;//BAIL
-					}
-					VVMSer_insEnd(Context, (int)(VVMSerCtxStartInsCode - VVMSerStartCtxCode), (int)(Context.Code - VVMSerCtxStartInsCode));
-				} while (Op != EVectorVMOp::done);
-				VVMSer_chunkEnd(SerializeState)
-			}
-
-			InstancesLeft -= GParallelVVMInstancesPerChunk;
-			++ChunkIdx;
-		}
-		Context.FinishExec();
-	};
-
-	if ( NumBatches > 1 )
-	{
-		ParallelFor(NumBatches, ExecChunkBatch, GbParallelVVM == 0 || !bParallel);
-	}
-	else
-	{
-		ExecChunkBatch(0);
-	}
-
-#if UE_BUILD_TEST
-	if (bNumInstancesEvent)
-	{
-		FPlatformMisc::EndNamedEvent();
-	}
-#endif
-
-#if VECTORVM_SUPPORTS_SERIALIZATION
-	uint64 EndTime = FPlatformTime::Cycles64();
-	if (SerializeState) {
-		SerializeState->ExecDt = EndTime - StartTime; //NOTE: doesn't work if ParallelFor splits the work into multiple threads
-	}
-#endif //VVM_INCLUDE_SERIALIZATION
-}
-
-void ExecBatchedOutput(FVectorVMContext& Context)
-{
-	while (true)
-	{
-		FVectorVMExecFunction ExecFunction = reinterpret_cast<FVectorVMExecFunction>(Context.DecodePtr());
-		if (ExecFunction == nullptr)
-		{
-			break;
-		}
-		ExecFunction(Context);
-	}
-}
-
-void ConditionalAddOutputWrapper(FVectorVMCodeOptimizerContext& Context, bool& OuterAdded)
-{
-	if (!OuterAdded)
-	{
-		Context.WriteExecFunction(ExecBatchedOutput);
-		OuterAdded = true;
-	}
-}
-
-EVectorVMOp BatchedOutputOptimization(EVectorVMOp Op, FVectorVMCodeOptimizerContext& Context)
-{
-	if (!GbBatchVMOutput)
-	{
-		return Op;
-	}
-
-	bool OuterAdded = false;
-	bool HasValidOp = true;
-
-	while (HasValidOp)
-	{
-		switch (Op)
-		{
-		case EVectorVMOp::outputdata_half:
-			ConditionalAddOutputWrapper(Context, OuterAdded);
-			FScalarKernelWriteOutputIndexed<float, FFloat16, 2>::Optimize(Context);
-			break;
-
-		case EVectorVMOp::outputdata_int32:
-			ConditionalAddOutputWrapper(Context, OuterAdded);
-			FScalarKernelWriteOutputIndexed<int32, int32, 1>::Optimize(Context);
-			break;
-
-		case EVectorVMOp::outputdata_float:
-			ConditionalAddOutputWrapper(Context, OuterAdded);
-			FScalarKernelWriteOutputIndexed<float, float, 0>::Optimize(Context);
-			break;
-
-		default:
-			HasValidOp = false;
-			break;
-		}
-
-		if (HasValidOp)
-		{
-			Op = Context.BaseContext.DecodeOp();
-		}
-	}
-
-	if (OuterAdded)
-	{
-		Context.WriteExecFunction(nullptr);
-	}
-	return Op;
-}
-
-// Optimization managed by GbBatchPackVMOutput via PackedOutputOptimization()
-// Looks for the common pattern of an acquireindex op followed by a number of associated outputdata ops.  The
-// stock operation is to write an index into a temporary register, and then have the different outputs streams
-// write into the indexed location.  This optimization does a number of things:
-// -first we check if 'validity' is uniform or not, if it is we can have a fast path of both figuring out how many
-// indices we need, as well as how to write the output (if we find that they are all invalid, then we don't need to do anything!)
-// -if we need to evaluate the validity of each element we quickly count up the number (with vector intrinsics) and
-// grab a block of the indices (rather than one at a time)
-// -rather than storing the indices to use, we store a int8 mask which indicates a valid flag for each of the next 4 samples
-// -outputs are then written to depending on their source and their frequency:
-//		-uniform sources will be splatted to all valid entries
-//		-variable sources will be packed into the available slots
-struct FBatchedWriteIndexedOutput
-{
-	// acquires a batch of indices from the provided CounterHandler.  If we're running in parallel, then we'll need to use
-	// atomics to guarantee our place in the list of indices.
-	template<bool bParallel>
-	static VM_FORCEINLINE void AcquireCounterIndex(FVectorVMContext& Context, FDataSetCounterHandler& CounterHandler, int32 AcquireCount)
-	{
-		if (AcquireCount)
-		{
-			int32* CounterHandlerIndex = CounterHandler.Get();
-			int32 StartIndex = INDEX_NONE;
-
-			if (bParallel)
-			{
-				StartIndex = FPlatformAtomics::InterlockedAdd(CounterHandlerIndex, AcquireCount);
-			}
-			else
-			{
-				StartIndex = *CounterHandlerIndex;
-				*CounterHandlerIndex = StartIndex + AcquireCount;
-			}
-
-			// increment StartIndex, since CounterHandlerIndex starts at INDEX_NONE
-			Context.ValidInstanceIndexStart = StartIndex + 1;
-		}
-
-		Context.ValidInstanceCount = AcquireCount;
-		Context.ValidInstanceUniform = !AcquireCount || (Context.NumInstances == AcquireCount);
-	}
-
-	// evaluates a register to evaluate which instances are valid or not; will read 4 entries at a time and generate a
-	// a mask for which entries are valid as well as an overall count
-	template<bool bParallel>
-	static void HandleRegisterValidIndices(FVectorVMContext& Context)
-	{
-		FDataSetCounterHandler CounterHandler(Context);
-		FRegisterHandler<VectorRegister4Float> ValidReader(Context);
-		FRegisterHandler<int8> Dst(Context);
-
-		int8* DestAddr = Dst.GetDest();
-
-		// we can process VECTOR_WIDTH_FLOATS entries at a time, generating a int8 mask for each set of 4 indicating
-		// which are valid
-		const int32 LoopCount = FMath::DivideAndRoundUp(Context.NumInstances, VECTOR_WIDTH_FLOATS);
-
-		int32 Remainder = Context.NumInstances;
-		int32 ValidCount = 0;
-		for (int32 LoopIt = 0; LoopIt < LoopCount; ++LoopIt)
-		{
-			// input register needs to be padded to allow for 16 byte reads; but mask out the ones beyond NumInstances
-			uint8 ValidMask = static_cast<uint8>(VectorMaskBits(ValidReader.GetAndAdvance()));
-			ValidMask &= ~(0xFF << FMath::Min(VECTOR_WIDTH_FLOATS, Remainder));
-			ValidCount += FMath::CountBits(ValidMask);
-			
-			DestAddr[LoopIt] = ValidMask;
-			Remainder -= VECTOR_WIDTH_FLOATS;
-		}
-
-		// grab our batch of indices
-		AcquireCounterIndex<bParallel>(Context, CounterHandler, ValidCount);
-	}
-
-	// evaluates the uniform check and grab the appropriate number of indices
-	template<typename ValidReaderType, bool bParallel>
-	static VM_FORCEINLINE void HandleUniformValidIndices(FVectorVMContext& Context)
-	{
-		FDataSetCounterHandler CounterHandler(Context);
-		ValidReaderType ValidReader(Context);
-
-		if (ValidReader.Get())
-		{
-			AcquireCounterIndex<bParallel>(Context, CounterHandler, Context.NumInstances);
-		}
-	}
-
-	template<uint8 SrcOpType>
-	static VM_FORCEINLINE void IndexExecOptimized(FVectorVMContext& Context)
-	{
-		if (Context.IsParallelExecution())
-		{
-			switch (SrcOpType)
-			{
-			case SRCOP_RRR: HandleRegisterValidIndices<true>(Context); break;
-			case SRCOP_RRC:	HandleUniformValidIndices<FConstantHandler<int32>, true>(Context); break;
-			default: check(0); break;
-			}
-		}
-		else
-		{
-			switch (SrcOpType)
-			{
-			case SRCOP_RRR: HandleRegisterValidIndices<false>(Context); break;
-			case SRCOP_RRC:	HandleUniformValidIndices<FConstantHandler<int32>, false>(Context); break;
-			default: check(0); break;
-			}
-		}
-	}
-
-	void OptimizeAcquireIndex(FVectorVMCodeOptimizerContext& Context)
-	{
-		const uint32 SrcOpType = Context.BaseContext.DecodeSrcOperandTypes();
-
-		AcquireIndexConstant = !!(SrcOpType & OP0_CONST);
-
-		switch (SrcOpType)
-		{
-		case SRCOP_RRR: Context.WriteExecFunction(IndexExecOptimized<SRCOP_RRR>); break;
-		case SRCOP_RRC: Context.WriteExecFunction(IndexExecOptimized<SRCOP_RRC>); break;
-		default: check(0); break;
-		}
-
-		DataSetCounterIndex = Context.DecodeU16();
-		ValidTestRegisterIndex = Context.DecodeU16();
-		WorkingRegisterIndex = Context.DecodeU16();
-
-		Context.Write(DataSetCounterIndex);
-		Context.Write(ValidTestRegisterIndex);
-
-		// we only need the working register if we've got non-uniform data
-		if (SrcOpType == SRCOP_RRR)
-		{
-			Context.Write(WorkingRegisterIndex);
-		}
-	}
-
-	bool IsOutput(EVectorVMOp Op) const
-	{
-		return (Op == EVectorVMOp::outputdata_int32
-			|| Op == EVectorVMOp::outputdata_float
-			|| Op == EVectorVMOp::outputdata_half);
-	}
-
-	bool IsValidEnd(EVectorVMOp Op) const
-	{
-		return (Op == EVectorVMOp::done
-			|| Op == EVectorVMOp::acquireindex);
-	}
-
-	static bool SkipIfEmpty(FVectorVMContext& Context)
-	{
-		if (!Context.ValidInstanceCount)
-		{
-			Context.SkipCode(2 * sizeof(uint16)); // don't need DataSetIndex or DestIndexRegisterIdx
-			const uint16 AccumulatedOpCount = Context.DecodeU16();
-
-			Context.SkipCode(AccumulatedOpCount * 2 * sizeof(uint16));
-			return true;
-		}
-
-		return false;
-	}
-
-	VM_FORCEINLINE static void SplatElement(int32 OutputCount, int32 ElementValue, int32* RESTRICT OutputElements)
-	{
-		const int32 OutputCountWide = AlignDown(OutputCount, VECTOR_WIDTH_FLOATS);
-
-		if (OutputCountWide)
-		{
-			const VectorRegister4Int SplatValue = MakeVectorRegisterInt(ElementValue, ElementValue, ElementValue, ElementValue);
-
-			for (int32 i = 0; i < OutputCountWide; i += VECTOR_WIDTH_FLOATS)
-			{
-				VectorIntStore(SplatValue, OutputElements + i);
-			}
-		}
-
-		for (int32 i = OutputCountWide; i < OutputCount; ++i)
-		{
-			OutputElements[i] = ElementValue;
-		}
-	}
-
-	VM_FORCEINLINE static void SplatElement(int32 OutputCount, float ElementValue, float* RESTRICT OutputElements)
-	{
-		const int32 OutputCountWide = AlignDown(OutputCount, VECTOR_WIDTH_FLOATS);
-
-		if (OutputCountWide)
-		{
-			const VectorRegister4Float SplatValue = MakeVectorRegister(ElementValue, ElementValue, ElementValue, ElementValue);
-
-			for (int32 i = 0; i < OutputCountWide; i += VECTOR_WIDTH_FLOATS)
-			{
-				VectorStore(SplatValue, OutputElements + i);
-			}
-		}
-
-		for (int32 i = OutputCountWide; i < OutputCount; ++i)
-		{
-			OutputElements[i] = ElementValue;
-		}
-	}
-
-
-	VM_FORCEINLINE static void SplatElement(int32 OutputCount, float ElementValue, FFloat16* RESTRICT OutputElements)
-	{
-		uint16 TargetValue;
-		FPlatformMath::StoreHalf(&TargetValue, ElementValue);
-		uint32 TargetValue32 = uint32(TargetValue) | ((uint32(TargetValue) << 16));
-		const VectorRegister4Float SplatValue = MakeVectorRegister(TargetValue32, TargetValue32, TargetValue32, TargetValue32);
-
-		while (OutputCount > 7)
-		{
-			VectorStore(SplatValue, (float*)OutputElements); // TODO: LWC: this was using a void* previously which always assumed float*, but revisit, doesn't seem correct either way...
-			OutputElements += 8;
-			OutputCount -= 8;
-		}
-
-		while (OutputCount > 0)
-		{
-			*((uint16*)OutputElements) = TargetValue;
-			OutputElements++;
-			OutputCount--;
-		}
-	}
-
-	template<typename T>
-	VM_FORCEINLINE static void CopyElements(int32 ElementCount, const T* RESTRICT SourceElements, T* RESTRICT OutputElements)
-	{
-		memcpy(OutputElements, SourceElements, ElementCount * sizeof(T));
-	}
-
-	VM_FORCEINLINE static void CopyElements(int32 ElementCount, const float* RESTRICT SourceElements, FFloat16* RESTRICT OutputElements)
-	{
-		const float* Src = SourceElements;
-		uint16* Dst = (uint16*)OutputElements;
-		while (ElementCount > 7)
-		{
-			FPlatformMath::WideVectorStoreHalf(Dst, Src);
-			Dst += 8;
-			Src += 8;
-			ElementCount -= 8;
-		}
-		while (ElementCount > 3)
-		{
-			FPlatformMath::VectorStoreHalf(Dst, Src);
-			Dst += 4;
-			Src += 4;
-			ElementCount -= 4;
-		}
-		while (ElementCount > 0)
-		{
-			FPlatformMath::StoreHalf(Dst, *Src);
-			Dst += 1;
-			Src += 1;
-			ElementCount -= 1;
-		}
-	}
-
-	template<typename SourceType, typename TargetType>
-	VM_FORCEINLINE static void ScalarShuffleElements(int32 ElementCount, const SourceType* RESTRICT SourceElements, const int8* RESTRICT ValidMask, TargetType* RESTRICT OutputElements)
-	{
-		constexpr int32 ElementsPerMask = 4;
-
-		// scalar path that will read 4 values and write one at a time to the output based on the valid mask
-		int32 MaskIt = 0;
-
-		while (ElementCount)
-		{
-			const int8 ShuffleMask = ValidMask[MaskIt];
-			const int8 AdvanceCount = FMath::CountBits(ShuffleMask);
-			check(AdvanceCount >= 0 && AdvanceCount <= 4);
-			if (AdvanceCount)
-			{
-				for (int32 ScalarIt = 0; ScalarIt < ElementsPerMask; ++ScalarIt)
-				{
-					if (!!(ShuffleMask & (1 << ScalarIt)))
-					{
-						*OutputElements = SourceElements[MaskIt * ElementsPerMask + ScalarIt];
-						++OutputElements;
-					}
-				}
-
-				ElementCount -= AdvanceCount;
-			}
-			++MaskIt;
-		}
-	}
-
-	VM_FORCEINLINE static void ScalarShuffleElements(int32 ElementCount, const float* RESTRICT SourceElements, const int8* RESTRICT ValidMask, FFloat16* RESTRICT OutputElements)
-	{
-		constexpr int32 ElementsPerMask = 4;
-
-		// scalar path that will read 4 values and write one at a time to the output based on the valid mask
-		int32 MaskIt = 0;
-
-		while (ElementCount)
-		{
-			checkSlow(ElementCount > 0);
-			const int8 ShuffleMask = ValidMask[MaskIt];
-			const int8 AdvanceCount = FMath::CountBits(ShuffleMask);
-			if (AdvanceCount)
-			{
-				for (int32 ScalarIt = 0; ScalarIt < ElementsPerMask; ++ScalarIt)
-				{
-					if (!!(ShuffleMask & (1 << ScalarIt)))
-					{
-						FPlatformMath::StoreHalf((uint16*)OutputElements, SourceElements[MaskIt * ElementsPerMask + ScalarIt]);
-						++OutputElements;
-					}
-				}
-
-				ElementCount -= AdvanceCount;
-			}
-			++MaskIt;
-		}
-	}
-
-	VM_FORCEINLINE static void VectorShuffleElements(int32 ElementCount, const int32* RESTRICT SourceElements, const int8* RESTRICT ValidMask, int32* RESTRICT OutputElements)
-	{
-		int32 SourceIt = 0;
-
-		const VectorRegister4Int* RESTRICT SourceVectors = reinterpret_cast<const VectorRegister4Int* RESTRICT>(SourceElements);
-
-		// vector shuffle path writes 4 at a time (though the trailing elements may not be valid) until we have to move over
-		// to the scalar version for fear of overwriting our neighbors
-		while (ElementCount >= VECTOR_WIDTH_FLOATS)
-		{
-			const int8 ShuffleMask = ValidMask[SourceIt];
-			const int8 AdvanceCount = FMath::CountBits(ShuffleMask);
-			check(AdvanceCount >= 0 && AdvanceCount <= 4);
-
-			//
-			// VectorIntStore(		- unaligned writes of 16 bytes to our Destination; note that this maneuver requires us to have
-			//						our output buffers padded out to 16 bytes!
-			//	VectorIntShuffle(	- swizzle our source register to pack the valid entries at the beginning, with 0s at the end
-			//    Source,			- source data
-			//    ShuffleMask),		- result of the VectorMaskBits done in the acquireindex, int8/VectorRegister4Float of input
-			//  Destination);
-			VectorIntStore(VectorIntShuffle(SourceVectors[SourceIt], VectorVMConstants::RegisterShuffleMask[ShuffleMask]), OutputElements);
-
-			OutputElements += AdvanceCount;
-			ElementCount -= AdvanceCount;
-
-			++SourceIt;
-		}
-
-		ScalarShuffleElements(ElementCount, SourceElements + VECTOR_WIDTH_FLOATS * SourceIt, ValidMask + SourceIt, OutputElements);
-	}
-
-	VM_FORCEINLINE static void VectorShuffleElements(int32 ElementCount, const float* RESTRICT SourceElements, const int8* RESTRICT ValidMask, FFloat16* RESTRICT OutputElements)
-	{
-		int32 SourceIt = 0;
-
-		const VectorRegister4Int* RESTRICT SourceVectors = reinterpret_cast<const VectorRegister4Int* RESTRICT>(SourceElements);
-
-		// vector shuffle path writes 4 at a time (though the trailing elements may not be valid) until we have to move over
-		// to the scalar version for fear of overwriting our neighbors
-		while (ElementCount >= VECTOR_WIDTH_FLOATS)
-		{
-			const int8 ShuffleMask = ValidMask[SourceIt];
-			const int8 AdvanceCount = FMath::CountBits(ShuffleMask);
-			check(AdvanceCount >= 0 && AdvanceCount <= 4);
-
-			// VectorIntStore(		- unaligned writes of 16 bytes to our Destination; note that this maneuver requires us to have
-			//						our output buffers padded out to 16 bytes!
-			//	VectorIntShuffle(	- swizzle our source register to pack the valid entries at the beginning, with 0s at the end
-			//    Source,			- source data
-			//    ShuffleMask),		- result of the VectorMaskBits done in the acquireindex, int8/VectorRegister4Float of input
-			//  Destination);
-
-			const VectorRegister4Int ShuffledFloats = VectorIntShuffle(SourceVectors[SourceIt], VectorVMConstants::RegisterShuffleMask[ShuffleMask]);
-			FPlatformMath::VectorStoreHalf((uint16*)OutputElements, (float*)&ShuffledFloats);
-
-			OutputElements += AdvanceCount;
-			ElementCount -= AdvanceCount;
-
-			++SourceIt;
-		}
-
-		ScalarShuffleElements(ElementCount, SourceElements + VECTOR_WIDTH_FLOATS * SourceIt, ValidMask + SourceIt, OutputElements);
-	}
-
-	template<typename SourceType, typename TargetType, int32 TypeOffset>
-	static void CopyConstantToOutput(FVectorVMContext& Context)
-	{
-		if (SkipIfEmpty(Context))
-		{
-			return;
-		}
-
-		const uint16 DataSetIndex = Context.DecodeU16();
-		Context.SkipCode(sizeof(uint16)); // don't need DestIndexRegisterIdx
-		const uint16 AccumulatedOpCount = Context.DecodeU16();
-
-		FDataSetMeta& DataSetMeta = Context.GetDataSetMeta(DataSetIndex);
-
-		for (uint16 OpIt = 0; OpIt < AccumulatedOpCount; ++OpIt)
-		{
-			FConstantHandler<SourceType> SourceHandler(Context);
-			const uint16 DestRegisterIndex = Context.DecodeU16();
-
-			TargetType* RESTRICT OutputElements = Context.GetOutputRegister<TargetType, TypeOffset>(DataSetIndex, DestRegisterIndex) + Context.ValidInstanceIndexStart;
-
-			SplatElement(Context.ValidInstanceCount, SourceHandler.Constant, OutputElements);
-		}
-	}
-
-	template<typename SourceType, typename TargetType, int32 TypeOffset, bool bCheckIfEmpty = true>
-	static void CopyRegisterToOutput(FVectorVMContext& Context)
-	{
-		if (bCheckIfEmpty && SkipIfEmpty(Context))
-		{
-			return;
-		}
-
-		const uint16 DataSetIndex = Context.DecodeU16();
-		Context.SkipCode(sizeof(uint16)); // don't need DestIndexRegisterIdx
-		const uint16 AccumulatedOpCount = Context.DecodeU16();
-
-		check(Context.ValidInstanceCount == Context.NumInstances);
-		FDataSetMeta& DataSetMeta = Context.GetDataSetMeta(DataSetIndex);
-
-		for (uint16 OpIt = 0; OpIt < AccumulatedOpCount; ++OpIt)
-		{
-			FRegisterHandler<SourceType> SourceHandler(Context);
-			const uint16 DestRegisterIndex = Context.DecodeU16();
-
-			TargetType* RESTRICT OutputElements = Context.GetOutputRegister<TargetType, TypeOffset>(DataSetIndex, DestRegisterIndex) + Context.ValidInstanceIndexStart;
-
-			CopyElements(Context.ValidInstanceCount, SourceHandler.GetDest(), OutputElements);
-		}
-	}
-
-	template<typename SourceType, typename TargetType, int32 TypeOffset>
-	static void ShuffleRegisterToOutput(FVectorVMContext& Context)
-	{
-		if (SkipIfEmpty(Context))
-		{
-			return;
-		}
-		else if (Context.ValidInstanceUniform)
-		{
-			CopyRegisterToOutput<SourceType, TargetType, TypeOffset, false>(Context);
-			return;
-		}
-
-		const uint16 DataSetIndex = Context.DecodeU16();
-		const uint16 DestIndexRegisterIdx = Context.DecodeU16();
-		const uint16 AccumulatedOpCount = Context.DecodeU16();
-
-		FDataSetMeta& DataSetMeta = Context.GetDataSetMeta(DataSetIndex);
-
-		const int8* RESTRICT DestIndexReg = reinterpret_cast<const int8*>(Context.GetTempRegister(DestIndexRegisterIdx));
-
-		for (uint16 OpIt = 0; OpIt < AccumulatedOpCount; ++OpIt)
-		{
-			FRegisterHandler<SourceType> SourceRegister(Context);
-			TargetType* RESTRICT DestReg = Context.GetOutputRegister<TargetType, TypeOffset>(DataSetIndex, Context.DecodeU16()) + Context.ValidInstanceIndexStart;
-
-			// the number of instances that we're expecting to write.  it is important that we keep track of it because when we
-			// get down to the end we need to switch from the shuffled approach to a scalar approach so that we don't
-			// overwrite the indexed output that another parallel context might have written to
-			VectorShuffleElements(Context.ValidInstanceCount, SourceRegister.GetDest(), DestIndexReg, DestReg);
-		}
-	}
-
-	bool OptimizeBatch(FVectorVMCodeOptimizerContext& Context)
-	{
-		const int32 BatchedOpCount = BatchedOps.Num();
-
-		if (!BatchedOpCount)
-			return false;
-
-		for (const auto& BatchEntry : BatchedOps)
-		{
-			const uint16 AccumulatedOpCount = BatchEntry.Value.Num();
-
-			if (!AccumulatedOpCount)
-				continue;
-
-			// matrix of options based on the:
-			//	Op (are we outputting ints, floats or halfs)
-			//	SrcOpType (are we copying a constant or a register over)
-			// additionally, if we know that the index is constant, then we know at worst we're doing a copy
-
-			if (BatchEntry.Key.SrcOpType == SRCOP_RRC)
-			{
-				switch (BatchEntry.Key.Op)
-				{
-				case EVectorVMOp::outputdata_float: Context.WriteExecFunction(CopyConstantToOutput<float, float, 0>); break;
-				case EVectorVMOp::outputdata_int32: Context.WriteExecFunction(CopyConstantToOutput<int32, int32, 1>); break;
-				case EVectorVMOp::outputdata_half: Context.WriteExecFunction(CopyConstantToOutput<float, FFloat16, 2>); break;
-				default: check(0);
-				}
-			}
-			else if (AcquireIndexConstant)
-			{
-				switch (BatchEntry.Key.Op)
-				{
-				case EVectorVMOp::outputdata_float: Context.WriteExecFunction(CopyRegisterToOutput<float, float, 0>); break;
-				case EVectorVMOp::outputdata_int32: Context.WriteExecFunction(CopyRegisterToOutput<int32, int32, 1>); break;
-				case EVectorVMOp::outputdata_half: Context.WriteExecFunction(CopyRegisterToOutput<float, FFloat16, 2>); break;
-				default: check(0);
-				}
-			}
-			else
-			{
-				check(BatchEntry.Key.SrcOpType == SRCOP_RRR);
-				switch (BatchEntry.Key.Op)
-				{
-				case EVectorVMOp::outputdata_float: Context.WriteExecFunction(ShuffleRegisterToOutput<int32, int32, 0>); break;
-				case EVectorVMOp::outputdata_int32: Context.WriteExecFunction(ShuffleRegisterToOutput<int32, int32, 1>); break;
-				case EVectorVMOp::outputdata_half: Context.WriteExecFunction(ShuffleRegisterToOutput<float, FFloat16, 2>); break;
-				default: check(0);
-				}
-			}
-
-			Context.Write(BatchEntry.Key.DataSetIndex);
-			Context.Write(BatchEntry.Key.DestIndexRegisterIdx);
-			Context.Write(AccumulatedOpCount);
-			for (const FOpValue& OpValue : BatchEntry.Value)
-			{
-				Context.Write(OpValue.SourceRegisterIndex);
-				Context.Write(OpValue.DestRegisterIdx);
-			}
-		}
-
-		return true;
-	}
-
-	bool ExtractOp(EVectorVMOp Op, FVectorVMCodeOptimizerContext& Context)
-	{
-		FOpKey Key;
-		Key.Op = Op;
-		Key.SrcOpType = Context.BaseContext.DecodeSrcOperandTypes();
-		Key.DataSetIndex = Context.DecodeU16();
-		Key.DestIndexRegisterIdx = Context.DecodeU16();
-
-		if (Key.DestIndexRegisterIdx != WorkingRegisterIndex)
-		{
-			// if we've found an output node that is not related to the acquire index op, then just exit
-			return false;
-		}
-
-		FOpValue Value;
-		Value.SourceRegisterIndex = Context.DecodeU16();
-		Value.DestRegisterIdx = Context.DecodeU16();
-
-		TArray<FOpValue>& ExistingOps = BatchedOps.FindOrAdd(Key);
-		ExistingOps.Add(Value);
-
-		return true;
-	}
-
-private:
-	using RegisterType = VectorRegister4Int;
-	using ScalarType = int32;
-
-	uint16 DataSetCounterIndex = 0;
-	uint16 ValidTestRegisterIndex = 0;
-	uint16 WorkingRegisterIndex = 0;
-	bool AcquireIndexConstant = false;
-
-	struct FOpKey
-	{
-		uint16 DestIndexRegisterIdx;
-		uint16 DataSetIndex;
-		uint8 SrcOpType;
-		EVectorVMOp Op;
-	};
-
-	struct FOpValue
-	{
-		uint16 SourceRegisterIndex;
-		uint16 DestRegisterIdx;
-	};
-
-	struct FOpKeyFuncs : public TDefaultMapKeyFuncs<FOpKey, TArray<FOpValue>, false>
-	{
-		static VM_FORCEINLINE bool Matches(const FOpKey& A, const FOpKey& B)
-		{
-			return A.DestIndexRegisterIdx == B.DestIndexRegisterIdx
-				&& A.DataSetIndex == B.DataSetIndex
-				&& A.SrcOpType == B.SrcOpType
-				&& A.Op == B.Op;
-		}
-
-		static VM_FORCEINLINE uint32 GetKeyHash(const FOpKey& Key)
-		{
-			return HashCombine(Key.DestIndexRegisterIdx | (Key.DataSetIndex << 16), Key.SrcOpType | (static_cast<uint8>(Key.Op) << 8));
-		}
-	};
-
-	TMap<FOpKey, TArray<FOpValue>, FDefaultSetAllocator, FOpKeyFuncs> BatchedOps;
-};
-
-// look for the pattern of acquireindex followed by a bunch of outputs.
-bool PackedOutputOptimization(EVectorVMOp Op, FVectorVMCodeOptimizerContext& Context)
-{
-	if (!GbBatchPackVMOutput)
-	{
-		return false;
-	}
-
-	if (Op == EVectorVMOp::acquireindex)
-	{
-		const auto RollbackState = Context.CreateCodeState();
-
-		FBatchedWriteIndexedOutput BatchedOutputOp;
-
-		BatchedOutputOp.OptimizeAcquireIndex(Context);
-
-		bool BatchValid = true;
-
-		EVectorVMOp NextOp = Context.PeekOp();
-		while (BatchValid && BatchedOutputOp.IsOutput(NextOp))
-		{
-			BatchValid = BatchedOutputOp.ExtractOp(Context.BaseContext.DecodeOp(), Context);
-			NextOp = Context.PeekOp();
-		}
-
-		// if there's nothing worth optimizing here, then just revert what we've parsed
-		if (!BatchValid || !BatchedOutputOp.OptimizeBatch(Context))
-		{
-			Context.RollbackCodeState(RollbackState);
-			return false;
-		}
-
-		// We handled the existing op
-		return true;
-	}
-
-	return false;
-}
-
-void ExecBatchedInput(FVectorVMContext& Context)
-{
-	while (true)
-	{
-		FVectorVMExecFunction ExecFunction = reinterpret_cast<FVectorVMExecFunction>(Context.DecodePtr());
-		if (ExecFunction == nullptr)
-		{
-			break;
-		}
-		ExecFunction(Context);
-	}
-}
-
-void ConditionalAddInputWrapper(FVectorVMCodeOptimizerContext& Context, bool& OuterAdded)
-{
-	if (!OuterAdded)
-	{
-		Context.WriteExecFunction(ExecBatchedInput);
-		OuterAdded = true;
-	}
-}
-
-EVectorVMOp BatchedInputOptimization(EVectorVMOp Op, FVectorVMCodeOptimizerContext& Context)
-{
-	if (!GbBatchVMInput)
-	{
-		return Op;
-	}
-
-	bool OuterAdded = false;
-	bool HasValidOp = true;
-
-	while (HasValidOp)
-	{
-		switch (Op)
-		{
-		case EVectorVMOp::inputdata_half:
-			ConditionalAddInputWrapper(Context, OuterAdded);
-			FVectorKernelReadInput<FFloat16, 2>::Optimize(Context);
-			break;
-
-		case EVectorVMOp::inputdata_int32:
-			ConditionalAddInputWrapper(Context, OuterAdded);
-			FVectorKernelReadInput<int32, 1>::Optimize(Context);
-			break;
-
-		case EVectorVMOp::inputdata_float:
-			ConditionalAddInputWrapper(Context, OuterAdded);
-			FVectorKernelReadInput<float, 0>::Optimize(Context);
-			break;
-
-		default:
-			HasValidOp = false;
-			break;
-		}
-
-		if (HasValidOp)
-		{
-			Op = Context.BaseContext.DecodeOp();
-		}
-	}
-
-	if (OuterAdded)
-	{
-		Context.WriteExecFunction(nullptr);
-	}
-	return Op;
-}
-
-bool SafeMathOptimization(EVectorVMOp Op, FVectorVMCodeOptimizerContext& Context)
-{
-	if (!GbSafeOptimizedKernels)
-	{
-		return false;
-	}
-
-	switch (Op)
-	{
-		case EVectorVMOp::div: FVectorKernelDivSafe::Optimize(Context); return true;
-		case EVectorVMOp::rcp: FVectorKernelRcpSafe::Optimize(Context); return true;
-		case EVectorVMOp::rsq: FVectorKernelRsqSafe::Optimize(Context); return true;
-		case EVectorVMOp::sqrt: FVectorKernelSqrtSafe::Optimize(Context); return true;
-		case EVectorVMOp::log: FVectorKernelLogSafe::Optimize(Context); return true;
-		case EVectorVMOp::pow: FVectorKernelPowSafe::Optimize(Context); return true;
-		default:
-			return false;
-	}
-}
-
-void VectorVM::OptimizeByteCode(const uint8* ByteCode, TArray<uint8>& OptimizedCode, TArrayView<uint8> ExternalFunctionRegisterCounts)
-{
-	OptimizedCode.Empty();
-
-	//-TODO:7 Support unaligned writes & little endian
-#if PLATFORM_SUPPORTS_UNALIGNED_LOADS && PLATFORM_LITTLE_ENDIAN
-
-	if (!GbOptimizeVMByteCode || (ByteCode == nullptr))
-	{
-		return;
-	}
-
-	FVectorVMCodeOptimizerContext Context(FVectorVMContext::Get(), ByteCode, OptimizedCode, ExternalFunctionRegisterCounts);
-
-	// add any optimization filters in here, useful so what we can isolate optimizations with CVars
-	FVectorVMCodeOptimizerContext::OptimizeVMFunction VMFilters[] =
-	{
-		//BatchedInputOptimization,
-		//BatchedOutputOptimization,
-		PackedOutputOptimization,
-		SafeMathOptimization,
-	};
-
-	EVectorVMOp Op = EVectorVMOp::done;
-	do
-	{
-		Op = Context.BaseContext.DecodeOp();
-
-		// Filters allow us to modify a single or series of operations
-		bool bOpWasFiltered = false;
-		for (auto Filter : VMFilters)
-		{
-			if ( Filter(Op, Context) )
-			{
-				bOpWasFiltered = true;
-				break;
-			}
-		}
-		if (bOpWasFiltered)
-		{
-			continue;
-		}
-
-		// Optimize op
-		switch (Op)
-		{
-			case EVectorVMOp::add: FVectorKernelAdd::Optimize(Context); break;
-			case EVectorVMOp::sub: FVectorKernelSub::Optimize(Context); break;
-			case EVectorVMOp::mul: FVectorKernelMul::Optimize(Context); break;
-			case EVectorVMOp::div: FVectorKernelDiv::Optimize(Context); break;
-			case EVectorVMOp::mad: FVectorKernelMad::Optimize(Context); break;
-			case EVectorVMOp::lerp: FVectorKernelLerp::Optimize(Context); break;
-			case EVectorVMOp::rcp: FVectorKernelRcp::Optimize(Context); break;
-			case EVectorVMOp::rsq: FVectorKernelRsq::Optimize(Context); break;
-			case EVectorVMOp::sqrt: FVectorKernelSqrt::Optimize(Context); break;
-			case EVectorVMOp::neg: FVectorKernelNeg::Optimize(Context); break;
-			case EVectorVMOp::abs: FVectorKernelAbs::Optimize(Context); break;
-			case EVectorVMOp::exp: FVectorKernelExp::Optimize(Context); break;
-			case EVectorVMOp::exp2: FVectorKernelExp2::Optimize(Context); break;
-			case EVectorVMOp::log: FVectorKernelLog::Optimize(Context); break;
-			case EVectorVMOp::log2: FVectorKernelLog2::Optimize(Context); break;
-			case EVectorVMOp::sin: FVectorKernelSin::Optimize(Context); break;
-			case EVectorVMOp::cos: FVectorKernelCos::Optimize(Context); break;
-			case EVectorVMOp::tan: FVectorKernelTan::Optimize(Context); break;
-			case EVectorVMOp::asin: FVectorKernelASin::Optimize(Context); break;
-			case EVectorVMOp::acos: FVectorKernelACos::Optimize(Context); break;
-			case EVectorVMOp::atan: FVectorKernelATan::Optimize(Context); break;
-			case EVectorVMOp::atan2: FVectorKernelATan2::Optimize(Context); break;
-			case EVectorVMOp::ceil: FVectorKernelCeil::Optimize(Context); break;
-			case EVectorVMOp::floor: FVectorKernelFloor::Optimize(Context); break;
-			case EVectorVMOp::round: FVectorKernelRound::Optimize(Context); break;
-			case EVectorVMOp::fmod: FVectorKernelMod::Optimize(Context); break;
-			case EVectorVMOp::frac: FVectorKernelFrac::Optimize(Context); break;
-			case EVectorVMOp::trunc: FVectorKernelTrunc::Optimize(Context); break;
-			case EVectorVMOp::clamp: FVectorKernelClamp::Optimize(Context); break;
-			case EVectorVMOp::min: FVectorKernelMin::Optimize(Context); break;
-			case EVectorVMOp::max: FVectorKernelMax::Optimize(Context); break;
-			case EVectorVMOp::pow: FVectorKernelPow::Optimize(Context); break;
-			case EVectorVMOp::sign: FVectorKernelSign::Optimize(Context); break;
-			case EVectorVMOp::step: FVectorKernelStep::Optimize(Context); break;
-			case EVectorVMOp::random: FVectorKernelRandom::Optimize(Context); break;
-			case EVectorVMOp::noise: VectorVMNoise::Optimize_Noise1D(Context); break;
-			case EVectorVMOp::noise2D: VectorVMNoise::Optimize_Noise2D(Context); break;
-			case EVectorVMOp::noise3D: VectorVMNoise::Optimize_Noise3D(Context); break;
-
-			case EVectorVMOp::cmplt: FVectorKernelCompareLT::Optimize(Context); break;
-			case EVectorVMOp::cmple: FVectorKernelCompareLE::Optimize(Context); break;
-			case EVectorVMOp::cmpgt: FVectorKernelCompareGT::Optimize(Context); break;
-			case EVectorVMOp::cmpge: FVectorKernelCompareGE::Optimize(Context); break;
-			case EVectorVMOp::cmpeq: FVectorKernelCompareEQ::Optimize(Context); break;
-			case EVectorVMOp::cmpneq: FVectorKernelCompareNEQ::Optimize(Context); break;
-			case EVectorVMOp::select: FVectorKernelSelect::Optimize(Context); break;
-
-			case EVectorVMOp::addi: FVectorIntKernelAdd::Optimize(Context); break;
-			case EVectorVMOp::subi: FVectorIntKernelSubtract::Optimize(Context); break;
-			case EVectorVMOp::muli: FVectorIntKernelMultiply::Optimize(Context); break;
-			case EVectorVMOp::divi: FVectorIntKernelDivide::Optimize(Context); break;
-			case EVectorVMOp::clampi: FVectorIntKernelClamp::Optimize(Context); break;
-			case EVectorVMOp::mini: FVectorIntKernelMin::Optimize(Context); break;
-			case EVectorVMOp::maxi: FVectorIntKernelMax::Optimize(Context); break;
-			case EVectorVMOp::absi: FVectorIntKernelAbs::Optimize(Context); break;
-			case EVectorVMOp::negi: FVectorIntKernelNegate::Optimize(Context); break;
-			case EVectorVMOp::signi: FVectorIntKernelSign::Optimize(Context); break;
-			case EVectorVMOp::randomi: FScalarIntKernelRandom::Optimize(Context); break;
-			case EVectorVMOp::cmplti: FVectorIntKernelCompareLT::Optimize(Context); break;
-			case EVectorVMOp::cmplei: FVectorIntKernelCompareLE::Optimize(Context); break;
-			case EVectorVMOp::cmpgti: FVectorIntKernelCompareGT::Optimize(Context); break;
-			case EVectorVMOp::cmpgei: FVectorIntKernelCompareGE::Optimize(Context); break;
-			case EVectorVMOp::cmpeqi: FVectorIntKernelCompareEQ::Optimize(Context); break;
-			case EVectorVMOp::cmpneqi: FVectorIntKernelCompareNEQ::Optimize(Context); break;
-			case EVectorVMOp::bit_and: FVectorIntKernelBitAnd::Optimize(Context); break;
-			case EVectorVMOp::bit_or: FVectorIntKernelBitOr::Optimize(Context); break;
-			case EVectorVMOp::bit_xor: FVectorIntKernelBitXor::Optimize(Context); break;
-			case EVectorVMOp::bit_not: FVectorIntKernelBitNot::Optimize(Context); break;
-			case EVectorVMOp::bit_lshift: FVectorIntKernelBitLShift::Optimize(Context); break;
-			case EVectorVMOp::bit_rshift: FVectorIntKernelBitRShift::Optimize(Context); break;
-			case EVectorVMOp::logic_and: FVectorIntKernelLogicAnd::Optimize(Context); break;
-			case EVectorVMOp::logic_or: FVectorIntKernelLogicOr::Optimize(Context); break;
-			case EVectorVMOp::logic_xor: FVectorIntKernelLogicXor::Optimize(Context); break;
-			case EVectorVMOp::logic_not: FVectorIntKernelLogicNot::Optimize(Context); break;
-			case EVectorVMOp::f2i: FVectorKernelFloatToInt::Optimize(Context); break;
-			case EVectorVMOp::i2f: FVectorKernelIntToFloat::Optimize(Context); break;
-			case EVectorVMOp::f2b: FVectorKernelFloatToBool::Optimize(Context); break;
-			case EVectorVMOp::b2f: FVectorKernelBoolToFloat::Optimize(Context); break;
-			case EVectorVMOp::i2b: FVectorKernelIntToBool::Optimize(Context); break;
-			case EVectorVMOp::b2i: FVectorKernelBoolToInt::Optimize(Context); break;
-			case EVectorVMOp::fasi: FVectorKernelFloatAsInt::Optimize(Context); break;
-			case EVectorVMOp::iasf: FVectorKernelIntAsFloat::Optimize(Context); break;
-
-			case EVectorVMOp::outputdata_half:	FScalarKernelWriteOutputIndexed<float, FFloat16, 2>::Optimize(Context);	break;
-			case EVectorVMOp::inputdata_half: FVectorKernelReadInput<FFloat16, 2>::Optimize(Context); break;
-			case EVectorVMOp::outputdata_int32:	FScalarKernelWriteOutputIndexed<int32, int32, 1>::Optimize(Context);	break;
-			case EVectorVMOp::inputdata_int32: FVectorKernelReadInput<int32, 1>::Optimize(Context); break;
-			case EVectorVMOp::outputdata_float:	FScalarKernelWriteOutputIndexed<float, float, 0>::Optimize(Context);	break;
-			case EVectorVMOp::inputdata_float: FVectorKernelReadInput<float, 0>::Optimize(Context); break;
-			case EVectorVMOp::inputdata_noadvance_int32: FVectorKernelReadInputNoAdvance<int32, 1>::Optimize(Context); break;
-			case EVectorVMOp::inputdata_noadvance_float: FVectorKernelReadInputNoAdvance<float, 0>::Optimize(Context); break;
-			case EVectorVMOp::inputdata_noadvance_half: FVectorKernelReadInputNoAdvance<FFloat16, 2>::Optimize(Context); break;
-			case EVectorVMOp::acquireindex: FScalarKernelAcquireCounterIndex::Optimize(Context); break;
-			case EVectorVMOp::external_func_call: FKernelExternalFunctionCall::Optimize(Context); break;
-
-			case EVectorVMOp::exec_index: FVectorKernelExecutionIndex::Optimize(Context); break;
-
-			case EVectorVMOp::enter_stat_scope: FVectorKernelEnterStatScope::Optimize(Context); break;
-			case EVectorVMOp::exit_stat_scope: FVectorKernelExitStatScope::Optimize(Context); break;
-
-			//Special case ops to handle unique IDs but this can be written as generalized buffer operations. TODO!
-			case EVectorVMOp::update_id:	FScalarKernelUpdateID::Optimize(Context); break;
-			case EVectorVMOp::acquire_id:	FScalarKernelAcquireID::Optimize(Context); break;
-
-			// Execution always terminates with a "done" opcode.
-			case EVectorVMOp::done:
-				break;
-
-				// Opcode not recognized / implemented.
-			default:
-				UE_LOG(LogVectorVM, Fatal, TEXT("Unknown op code 0x%02x"), (uint32)Op);
-				OptimizedCode.Empty();
-				return;//BAIL
-		}
-	} while (Op != EVectorVMOp::done);
-	Context.WriteExecFunction(nullptr);
-
-	Context.EncodeJumpTable();
-#endif //PLATFORM_SUPPORTS_UNALIGNED_LOADS && PLATFORM_LITTLE_ENDIAN
-}
-
-#endif // #if VECTORVM_SUPPORTS_LEGACY
-
 #undef VM_FORCEINLINE
diff --git a/Engine/Source/Runtime/VectorVM/Private/VectorVMExperimental.cpp b/Engine/Source/Runtime/VectorVM/Private/VectorVMExperimental.cpp
--- a/Engine/Source/Runtime/VectorVM/Private/VectorVMExperimental.cpp
+++ b/Engine/Source/Runtime/VectorVM/Private/VectorVMExperimental.cpp
@@ -1,5 +1,4 @@
 // Copyright Epic Games, Inc. All Rights Reserved.
-#if VECTORVM_SUPPORTS_EXPERIMENTAL
 
 #include "VectorVMExperimental.h"
 #include "VectorVM.h"
@@ -1931,12 +1930,7 @@ const uint8 *VVM_external_func_call(const bool CT_MultipleLoops, const uint8 *In
 
 		ExtFnCtx.RandCounters             = &BatchState->ChunkLocalData.RandCounters;
 		ExtFnCtx.DataSets                 = ExecCtx->DataSets;
-#		if VECTORVM_SUPPORTS_LEGACY
-		FVectorVMExternalFunctionContextProxy ProxyExtFnCtx(ExtFnCtx);
-		ExtFnData->Function->Execute(ProxyExtFnCtx);
-#		else // VECTORVM_SUPPORTS_LEGACY
 		ExtFnData->Function->Execute(ExtFnCtx);
-#		endif //VECTORVM_SUPPORTS_LEGACY
 	}
 	return InsPtr + 3 + 2 * (ExtFnData->NumInputs + ExtFnData->NumOutputs);
 }
@@ -2847,21 +2841,3 @@ uint16 VVMDbgGetRemappedConstIdx(FVectorVMState *VectorVMState, int ConstTableId
 	return 0xFFFF;
 }
 
-#else //NIAGARA_EXP_VM
-
-struct FVectorVMState *AllocVectorVMState(struct FVectorVMOptimizeContext *OptimizeCtx)
-{
-	return nullptr;
-}
-
-void FreeVectorVMState(struct FVectorVMState *VVMState)
-{
-
-}
-
-void ExecVectorVMState(struct FVectorVMState *VVMState, struct FVectorVMSerializeState *SerializeState, struct FVectorVMSerializeState *CmpSerializeState)
-{
-
-}
-
-#endif //NIAGARA_EXP_VM
diff --git a/Engine/Source/Runtime/VectorVM/Private/VectorVMExperimentalOptimizer.cpp b/Engine/Source/Runtime/VectorVM/Private/VectorVMExperimentalOptimizer.cpp
--- a/Engine/Source/Runtime/VectorVM/Private/VectorVMExperimentalOptimizer.cpp
+++ b/Engine/Source/Runtime/VectorVM/Private/VectorVMExperimentalOptimizer.cpp
@@ -5,9 +5,6 @@
 #include "VectorVM.h"
 #include "VectorVMExperimental.h"
 
-
-#if VECTORVM_SUPPORTS_EXPERIMENTAL
-
 namespace VectorVMSerializationHelper
 {
 
@@ -3275,5 +3272,3 @@ void ReinterpretVectorVMOptimizeContextData(TConstArrayView<uint8> ContextData,
 {
 	VectorVMSerializationHelper::ThawContext(ContextData, Context);
 }
-
-#endif //VECTORVM_SUPPORTS_EXPERIMENTAL
diff --git a/Engine/Source/Runtime/VectorVM/Private/VectorVMExperimentalSerialization.cpp b/Engine/Source/Runtime/VectorVM/Private/VectorVMExperimentalSerialization.cpp
--- a/Engine/Source/Runtime/VectorVM/Private/VectorVMExperimentalSerialization.cpp
+++ b/Engine/Source/Runtime/VectorVM/Private/VectorVMExperimentalSerialization.cpp
@@ -141,9 +141,6 @@ static uint32 VectorVMSerializeSetError_(FVectorVMSerializeState *SerializeState
 
 #define VectorVMSerializeSetError(SerializeState, Flags)     VectorVMSerializeSetError_(SerializeState, Flags, __LINE__)
 
-
-#if VECTORVM_SUPPORTS_EXPERIMENTAL
-
 /*********************************************************************************************************************************************************************************************************************************
 *** Serialization stuff for the Experimental VM
 *********************************************************************************************************************************************************************************************************************************/
@@ -441,82 +438,6 @@ uint32 VVMSer_initSerializationState_(FVectorVMSerializeState *SerializeState, F
 }
 #endif //VVM_SERIALIZE_NO_WRITE
 
-#else // VECTORVM_SUPPORTS_EXPERIMENTAL
-
-/*********************************************************************************************************************************************************************************************************************************
-*** Serialization stuff for the old VM
-*********************************************************************************************************************************************************************************************************************************/
-#define VVMSer_batchStartExp(...)
-#define VVMSer_batchEndExp(...)
-
-#define VVMSer_chunkStartExp(...)
-#define VVMSer_chunkEndExp(...)
-
-#define VVMSer_insStartExp(...)
-#define VVMSer_insEndExp(...)
-
-#ifdef VVM_SERIALIZE_NO_WRITE
-#define VVMSer_chunkStart(...)
-#define VVMSer_chunkEnd(...)
-#define VVMSer_insStart(...)
-#define VVMSer_insEnd(...)
-#else
-#define VVMSer_chunkStart(Context, ChunkIdx, BatchIdx)                      \
-	FVectorVMSerializeChunk *SerializeChunk = nullptr;                      \
-	int VVMSerNumInstructionsThisChunk = 0;                                 \
-	unsigned char *VVMSerStartCtxCode = (unsigned char*)Context.Code;       \
-	if (SerializeState && SerializeState->Chunks)                           \
-	{                                                                       \
-		SerializeChunk = SerializeState->Chunks + ChunkIdx;                 \
-		SerializeChunk->ChunkIdx = ChunkIdx;                                \
-		SerializeChunk->BatchIdx = BatchIdx;                                \
-		SerializeChunk->ExecIdx  = BatchIdx;                                \
-		SerializeChunk->NumInstances = NumInstancesThisChunk;               \
-		SerializeChunk->StartThreadID = FPlatformTLS::GetCurrentThreadId(); \
-		SerializeChunk->StartClock = FPlatformTime::Cycles64();             \
-	}
-
-#define VVMSer_chunkEnd(SerializeState, ...)																            \
-	if (SerializeChunk)																						            \
-	{																										            \
-		FPlatformAtomics::InterlockedOr(&SerializeState->ChunkComplete, (1ULL << (uint64)SerializeChunk->ChunkIdx));    \
-		SerializeChunk->EndThreadID = FPlatformTLS::GetCurrentThreadId();										        \
-		SerializeChunk->EndClock = FPlatformTime::Cycles64();												            \
-	}
-
-#define VVMSer_insStart(Context, ...)										\
-	uint64 VVMSerStartCycles = FPlatformTime::Cycles64();					\
-	unsigned char *VVMSerCtxStartInsCode = (unsigned char*)Context.Code;	\
-
-#define VVMSer_insEnd(Context, OpStart, NumOps, ...)																																			\
-	if (SerializeState && Op != EVectorVMOp::done)                                                                                                                                              \
-	{																																												            \
-		uint64 VVMSerEndExecCycles = FPlatformTime::Cycles64();																																	\
-		VVMSer_serializeInstruction(SerializeState, Context, StartInstance, NumInstancesThisChunk, NumLoops, OpStart, NumOps, VVMSerStartCycles, VVMSerEndExecCycles, VVMSerEndExecCycles);		\
-		uint64 VVMSerEndSerializeCycles = FPlatformTime::Cycles64();																															\
-		SerializeState->SerializeDt += VVMSerEndSerializeCycles - VVMSerEndExecCycles;																											\
-		SerializeChunk->InsExecTime += VVMSerEndSerializeCycles - VVMSerEndExecCycles;																											\
-	}
-
-
-static void VVMSer_serializeInstruction(FVectorVMSerializeState *SerializeState, FVectorVMContext &Context, 
-										int StartInstance, int NumInstancesThisChunk, int NumLoops, 
-										int OpStart, int NumOps, uint64 StartInsCycles, uint64 EndDecodeCycles, uint64 EndInsCycles)
-{
-	FVectorVMSerializeInstruction *Ins = VVMSerGetNextInstruction(SerializeState, InstructionIdx, GlobalChunkIdx, OpStart, NumOps, EndInsCycles - StartInsCycles, EndDecodeCycles - StartInsCycles);
-	if (Ins)
-	{
-		for (uint32 i = 0; i < SerializeState->NumTempRegisters; ++i)
-		{
-			unsigned char *TempReg = (unsigned char *)Context.GetTempRegister(i);
-			FMemory::Memcpy(Ins->RegisterTable + SerializeState->NumInstances * i + StartInstance, TempReg, sizeof(uint32) * NumInstancesThisChunk);
-		}
-	}
-}
-#endif //VVM_SERIALIZE_NO_WRITE
-
-#endif // VECTORVM_SUPPORTS_EXPERIMENTAL
-
 /*********************************************************************************************************************************************************************************************************************************
 *** Stuff shared between the new and old VM
 *********************************************************************************************************************************************************************************************************************************/
@@ -768,7 +689,7 @@ void SerializeVectorVMWriteToFile(FVectorVMSerializeState *SerializeState, uint8
 			File->Write((uint8 *)&ExtData->NumInputs  , sizeof(uint16));
 			File->Write((uint8 *)&ExtData->NumOutputs , sizeof(uint16));
 		}
-#if VECTORVM_SUPPORTS_EXPERIMENTAL
+
 		//write the optimization context
 		if (SerializeState->Flags & VVMSer_IncludeOptContext) {
 			check(SerializeState->OptimizeCtx);
@@ -798,7 +719,6 @@ void SerializeVectorVMWriteToFile(FVectorVMSerializeState *SerializeState, uint8
 			}
 		}
 		delete File;
-#endif //VECTORVM_SUPPORTS_EXPERIMENTAL
 	}
 }
 
@@ -1078,7 +998,6 @@ uint32 SerializeVectorVMOutputDataSets(FVectorVMSerializeState *SerializeState,
 	return SerializeState->Error.Flags;
 }
 
-#if VECTORVM_SUPPORTS_EXPERIMENTAL
 uint32 SerializeVectorVMInputDataSets(FVectorVMSerializeState *SerializeState, FVectorVMExecContext *ExecCtx)
 {
 	return SerializeState ? SerializeVectorVMInputDataSets(SerializeState, ExecCtx->DataSets, ExecCtx->ConstantTableData, ExecCtx->ConstantTableNumBytes, ExecCtx->ConstantTableCount) : 0;
@@ -1088,7 +1007,6 @@ uint32 SerializeVectorVMOutputDataSets(FVectorVMSerializeState *SerializeState,
 {
 	return SerializeState ? SerializeVectorVMOutputDataSets(SerializeState, ExecCtx->DataSets, ExecCtx->ConstantTableData, ExecCtx->ConstantTableNumBytes, ExecCtx->ConstantTableCount) : 0;
 }
-#endif // #if VECTORVM_SUPPORTS_EXPERIMENTAL
 
 #else //VECTORVM_SUPPORTS_SERIALIZATION
 
diff --git a/Engine/Source/Runtime/VectorVM/Public/VectorVM.h b/Engine/Source/Runtime/VectorVM/Public/VectorVM.h
--- a/Engine/Source/Runtime/VectorVM/Public/VectorVM.h
+++ b/Engine/Source/Runtime/VectorVM/Public/VectorVM.h
@@ -13,71 +13,12 @@
 #include "Templates/AlignmentTemplates.h"
 #include "UObject/ObjectMacros.h"
 #include "VectorVMExperimental.h"
-#include "VectorVMLegacy.h"
 #include "VectorVMCommon.h"
 
 struct FVectorVMSerializeState;
 
-#if VECTORVM_SUPPORTS_EXPERIMENTAL && VECTORVM_SUPPORTS_LEGACY
-
-class FVectorVMExternalFunctionContextProxy
-{
-public:
-	const bool UsingExperimentalVM;
-	FVectorVMExternalFunctionContextExperimental Experimental;
-	FVectorVMExternalFunctionContextLegacy Legacy;
-
-	FVectorVMExternalFunctionContextProxy(const FVectorVMExternalFunctionContextExperimental& InExperimental)
-		: UsingExperimentalVM(true)
-		, Experimental(InExperimental)
-		, Legacy()
-	{}
-
-	FVectorVMExternalFunctionContextProxy(const FVectorVMExternalFunctionContextLegacy& InLegacy)
-		: UsingExperimentalVM(false)
-		, Experimental()
-		, Legacy(InLegacy)
-	{}
-
-	FORCEINLINE int32* GetRandCounters()
-	{
-		return UsingExperimentalVM ? Experimental.GetRandCounters() : Legacy.GetRandCounters().GetData();
-	}
-
-	FORCEINLINE FRandomStream& GetRandStream()
-	{
-		return UsingExperimentalVM ? Experimental.GetRandStream() : Legacy.GetRandStream();
-	}
-
-	FORCEINLINE int32 GetNumInstances() const
-	{
-		return UsingExperimentalVM ? Experimental.GetNumInstances() : Legacy.GetNumInstances();
-	}
-
-	template<uint32 InstancesPerOp>
-	FORCEINLINE int32 GetNumLoops() const
-	{
-		return UsingExperimentalVM ? Experimental.GetNumLoops<InstancesPerOp>() : Legacy.GetNumLoops<InstancesPerOp>();
-	}
-
-private:
-	FVectorVMExternalFunctionContextProxy() = delete;
-};
-
-using FVectorVMExternalFunctionContext = FVectorVMExternalFunctionContextProxy;
-
-#elif VECTORVM_SUPPORTS_EXPERIMENTAL
-
 using FVectorVMExternalFunctionContext = FVectorVMExternalFunctionContextExperimental;
 
-#elif VECTORVM_SUPPORTS_LEGACY
-
-using FVectorVMExternalFunctionContext = FVectorVMExternalFunctionContextLegacy;
-
-#else
-#error "At least one of VECTORVM_SUPPORTS_EXPERIMENTAL | VECTORVM_SUPPORTS_LEGACY must be defined"
-#endif
-
 namespace VectorVM
 {
 	/** Get total number of op-codes */
@@ -110,14 +51,6 @@ namespace VectorVM
 #endif
 	};
 
-#if VECTORVM_SUPPORTS_LEGACY
-	/**
-	 * Execute VectorVM bytecode.
-	 */
-	VECTORVM_API void Exec(FVectorVMExecArgs& Args, FVectorVMSerializeState *SerializeState);
-	VECTORVM_API void OptimizeByteCode(const uint8* ByteCode, TArray<uint8>& OptimizedCode, TArrayView<uint8> ExternalFunctionRegisterCounts);
-#endif
-
 	VECTORVM_API void Init();
 
 	#define VVM_EXT_FUNC_INPUT_LOC_BIT (unsigned short)(1<<15)
@@ -130,46 +63,12 @@ namespace VectorVM
 		T* Ptr;
 		FUserPtrHandler(FVectorVMExternalFunctionContext& Context)
 		{
-#if VECTORVM_SUPPORTS_EXPERIMENTAL && VECTORVM_SUPPORTS_LEGACY
-			if (Context.UsingExperimentalVM)
-			{
-				int32 AdvanceOffset;
-				int32* ConstPtr = (int32*)Context.Experimental.GetNextRegister(&AdvanceOffset);
-				check(AdvanceOffset == 0); //must be constant
-				UserPtrIdx = *ConstPtr;
-				check(UserPtrIdx != INDEX_NONE);
-				Ptr = (T*)Context.Experimental.GetUserPtrTable(UserPtrIdx);
-			}
-			else
-			{
-				const uint16 VariableOffset = Context.Legacy.DecodeU16();
-				check(!(VariableOffset& VVM_EXT_FUNC_INPUT_LOC_BIT));
-
-				const uint16 ConstantTableOffset = VariableOffset & VVM_EXT_FUNC_INPUT_LOC_MASK;
-				UserPtrIdx = *Context.Legacy.GetConstant<int32>(ConstantTableOffset);
-				check(UserPtrIdx != INDEX_NONE);
-
-				Ptr = static_cast<T*>(Context.Legacy.GetUserPtrTable(UserPtrIdx));
-			}
-#elif VECTORVM_SUPPORTS_EXPERIMENTAL
 			int32 AdvanceOffset;
 			int32* ConstPtr = (int32*)Context.GetNextRegister(&AdvanceOffset);
 			check(AdvanceOffset == 0); //must be constant
 			UserPtrIdx = *ConstPtr;
 			check(UserPtrIdx != INDEX_NONE);
 			Ptr = (T*)Context.GetUserPtrTable(UserPtrIdx);
-#elif VECTORVM_SUPPORTS_LEGACY
-			const uint16 VariableOffset = Context.DecodeU16();
-			check(!(VariableOffset & VVM_EXT_FUNC_INPUT_LOC_BIT));
-
-			const uint16 ConstantTableOffset = VariableOffset & VVM_EXT_FUNC_INPUT_LOC_MASK;
-			UserPtrIdx = *Context.GetConstant<int32>(ConstantTableOffset);
-			check(UserPtrIdx != INDEX_NONE);
-
-			Ptr = static_cast<T*>(Context.GetUserPtrTable(UserPtrIdx));
-#else 
-	#error "Not supported"
-#endif
 		}
 
 		FORCEINLINE T* Get() { return Ptr; }
@@ -201,40 +100,8 @@ namespace VectorVM
 
 		void Init(FVectorVMExternalFunctionContext& Context)
 		{
-#if VECTORVM_SUPPORTS_EXPERIMENTAL && VECTORVM_SUPPORTS_LEGACY
-			if (Context.UsingExperimentalVM)
-			{
-				InputPtr = (T*)Context.Experimental.GetNextRegister(&AdvanceOffset);
-				InputPtr += Context.Experimental.PerInstanceFnInstanceIdx * AdvanceOffset;
-			}
-			else
-			{
-				const int32 InputOffset = Context.Legacy.DecodeU16();
-				const bool bIsRegister = !!(InputOffset & VVM_EXT_FUNC_INPUT_LOC_BIT);
-				const int32 Offset = InputOffset & VVM_EXT_FUNC_INPUT_LOC_MASK;
-
-				AdvanceOffset = bIsRegister ? 1 : 0;
-				InputPtr = bIsRegister ? reinterpret_cast<T*>(Context.Legacy.GetTempRegister(Offset)) : Context.Legacy.GetConstant<T>(Offset);
-
-				//Hack: Offset into the buffer by the instance offset.
-				InputPtr += Context.Legacy.GetExternalFunctionInstanceOffset() * AdvanceOffset;
-			}
-#elif VECTORVM_SUPPORTS_EXPERIMENTAL
 			InputPtr = (T*)Context.GetNextRegister(&AdvanceOffset);
 			InputPtr += Context.PerInstanceFnInstanceIdx * AdvanceOffset;
-#elif VECTORVM_SUPPORTS_LEGACY
-			const int32 InputOffset = Context.DecodeU16();
-			const bool bIsRegister = !!(InputOffset & VVM_EXT_FUNC_INPUT_LOC_BIT);
-			const int32 Offset = InputOffset & VVM_EXT_FUNC_INPUT_LOC_MASK;
-
-			AdvanceOffset = bIsRegister ? 1 : 0;
-			InputPtr = bIsRegister ? reinterpret_cast<T*>(Context.GetTempRegister(Offset)) : Context.GetConstant<T>(Offset);
-
-			//Hack: Offset into the buffer by the instance offset.
-			InputPtr += Context.GetExternalFunctionInstanceOffset() * AdvanceOffset;
-#else
-	#error "Not supported"
-#endif
 
 			StartPtr = InputPtr;
 		}
@@ -266,40 +133,7 @@ namespace VectorVM
 	private:
 		T* RESTRICT Register = nullptr;
 		int32 AdvanceOffset = 0;
-#if VECTORVM_SUPPORTS_LEGACY
-		T Dummy;
-#endif
 	public:
-#if VECTORVM_SUPPORTS_EXPERIMENTAL && VECTORVM_SUPPORTS_LEGACY
-		FExternalFuncRegisterHandler(FVectorVMExternalFunctionContext& Context)
-		{
-			if (Context.UsingExperimentalVM)
-			{
-				Register = (T*)Context.Experimental.GetNextRegister(&AdvanceOffset);
-
-				//Hack: Offset into the buffer by the instance offset.
-				Register += Context.Experimental.PerInstanceFnInstanceIdx * AdvanceOffset;
-			}
-			else
-			{
-				const int32 RegisterIndex = Context.Legacy.DecodeU16() & VVM_EXT_FUNC_INPUT_LOC_MASK;
-				if (RegisterIndex == static_cast<int32>(VVM_EXT_FUNC_INPUT_LOC_MASK))
-				{
-					Register = &Dummy;
-					AdvanceOffset = 0;
-				}
-				else
-				{
-					checkSlow(RegisterIndex < Context.Legacy.GetNumTempRegisters());
-					Register = (T*)Context.Legacy.GetTempRegister(RegisterIndex);
-					AdvanceOffset = 1;
-
-					//Hack: Offset into the buffer by the instance offset.
-					Register += Context.Legacy.GetExternalFunctionInstanceOffset() * AdvanceOffset;
-				}
-			}
-		}
-#elif VECTORVM_SUPPORTS_EXPERIMENTAL
 		FExternalFuncRegisterHandler(FVectorVMExternalFunctionContext& Context)
 		{
 			Register = (T*)Context.GetNextRegister(&AdvanceOffset);
@@ -307,28 +141,6 @@ namespace VectorVM
 			//Hack: Offset into the buffer by the instance offset.
 			Register += Context.PerInstanceFnInstanceIdx * AdvanceOffset;
 		}
-#elif VECTORVM_SUPPORTS_LEGACY
-		FExternalFuncRegisterHandler(FVectorVMExternalFunctionContext& Context)
-		{
-			const int32 RegisterIndex = Context.DecodeU16() & VVM_EXT_FUNC_INPUT_LOC_MASK;
-			if (RegisterIndex == static_cast<int32>(VVM_EXT_FUNC_INPUT_LOC_MASK))
-			{
-				Register = &Dummy;
-				AdvanceOffset = 0;
-			}
-			else
-			{
-				checkSlow(RegisterIndex < Context.GetNumTempRegisters());
-				Register = (T*)Context.GetTempRegister(RegisterIndex);
-				AdvanceOffset = 1;
-
-				//Hack: Offset into the buffer by the instance offset.
-				Register += Context.GetExternalFunctionInstanceOffset() * AdvanceOffset;
-			}
-		}
-#else
-	#error "Not supported"
-#endif
 
 		FORCEINLINE bool IsValid() const { return AdvanceOffset > 0; }
 		FORCEINLINE const T Get() { return *Register; }
@@ -354,24 +166,6 @@ namespace VectorVM
 	{
 		T Constant;
 
-#if VECTORVM_SUPPORTS_EXPERIMENTAL && VECTORVM_SUPPORTS_LEGACY
-		FExternalFuncConstHandler(FVectorVMExternalFunctionContext& Context)
-		{
-			if (Context.UsingExperimentalVM)
-			{
-				int32 AdvanceOffset;
-				const T* Register = (const T*)Context.Experimental.GetNextRegister(&AdvanceOffset);
-				Register += Context.Experimental.PerInstanceFnInstanceIdx * AdvanceOffset;
-
-				Constant = *Register;
-			}
-			else
-			{
-				const uint16 ConstantIndex = (Context.Legacy.DecodeU16() & VVM_EXT_FUNC_INPUT_LOC_MASK);
-				Constant = (*Context.Legacy.GetConstant<T>(ConstantIndex));
-			}
-		}
-#elif VECTORVM_SUPPORTS_EXPERIMENTAL
 		FExternalFuncConstHandler(FVectorVMExternalFunctionContext& Context)
 		{
 			int32 AdvanceOffset;
@@ -380,11 +174,7 @@ namespace VectorVM
 
 			Constant = *Register;
 		}
-#elif VECTORVM_SUPPORTS_LEGACY
-		FExternalFuncConstHandler(FVectorVMExternalFunctionContext& Context)
-			: Constant(*Context.GetConstant<T>(Context.DecodeU16()& VVM_EXT_FUNC_INPUT_LOC_MASK))
-		{}
-#endif
+
 		FORCEINLINE const T& Get() { return Constant; }
 		FORCEINLINE const T& GetAndAdvance() { return Constant; }
 		FORCEINLINE void Advance() { }
diff --git a/Engine/Source/Runtime/VectorVM/Public/VectorVMCommon.h b/Engine/Source/Runtime/VectorVM/Public/VectorVMCommon.h
--- a/Engine/Source/Runtime/VectorVM/Public/VectorVMCommon.h
+++ b/Engine/Source/Runtime/VectorVM/Public/VectorVMCommon.h
@@ -286,22 +286,6 @@ struct FDataSetMeta
 
 	int32 IDAcquireTag;
 
-#if VECTORVM_SUPPORTS_LEGACY
-	//Temporary lock we're using for thread safety when writing to the FreeIDTable.
-	//TODO: A lock free algorithm is possible here. We can create a specialized lock free list and reuse the IDTable slots for FreeIndices as Next pointers for our LFL.
-	//This would also work well on the GPU. 
-	//UE-65856 for tracking this work.
-
-	//@NOTE (smcgrath): the lock/unlock functions can go with the ne VM, we don't need them
-	//Also worth noting: I implemented something similar to the above comment's TODO regarding a lock free algorithm.
-	//It's simplicity over the traditional approach was much appreciated, however in the end I just removed the whole thing
-	//because scripts were no longer being split across multiple threads, so there's no need for any syncronization anymore...
-	//if legacy is ever removed, all these comments can along with it.
-	FCriticalSection FreeTableLock;
-	FORCEINLINE void LockFreeTable();
-	FORCEINLINE void UnlockFreeTable();
-#endif
-
 	FDataSetMeta()
 		: InputRegisterTypeOffsets{}
 		, OutputRegisterTypeOffsets{}
@@ -357,22 +341,5 @@ private:
 typedef void * (VectorVMReallocFn)              (void *Ptr, size_t NumBytes, const char *Filename, int LineNumber);
 typedef void   (VectorVMFreeFn)                 (void *Ptr, const char *Filename, int LineNumber);
 
-#if VECTORVM_SUPPORTS_EXPERIMENTAL && VECTORVM_SUPPORTS_LEGACY
-
-class FVectorVMExternalFunctionContextProxy;
-DECLARE_DELEGATE_OneParam(FVMExternalFunction, FVectorVMExternalFunctionContextProxy& /*Context*/);
-
-#elif VECTORVM_SUPPORTS_EXPERIMENTAL
-
 class FVectorVMExternalFunctionContextExperimental;
 DECLARE_DELEGATE_OneParam(FVMExternalFunction, FVectorVMExternalFunctionContextExperimental& /*Context*/);
-
-#elif VECTORVM_SUPPORTS_LEGACY
-
-class FVectorVMExternalFunctionContextLegacy;
-DECLARE_DELEGATE_OneParam(FVMExternalFunction, FVectorVMExternalFunctionContextLegacy& /*Context*/);
-
-#else
-#error "At least one of VECTORVM_SUPPORTS_EXPERIMENTAL | VECTORVM_SUPPORTS_LEGACY must be defined"
-#endif
-
diff --git a/Engine/Source/Runtime/VectorVM/Public/VectorVMExperimental.h b/Engine/Source/Runtime/VectorVM/Public/VectorVMExperimental.h
--- a/Engine/Source/Runtime/VectorVM/Public/VectorVMExperimental.h
+++ b/Engine/Source/Runtime/VectorVM/Public/VectorVMExperimental.h
@@ -5,8 +5,6 @@
 #include "VectorVMCommon.h"
 #include "Math/RandomStream.h"
 
-#if VECTORVM_SUPPORTS_EXPERIMENTAL
-
 #define VECTORVM_SUPPORTS_COMPUTED_GOTO 0
 
 struct FDataSetMeta;
@@ -459,20 +457,3 @@ VECTORVM_API void             FreeVectorVMSerializeState             (FVectorVMS
 VECTORVM_API uint16           VVMDbgGetRemappedInputIdx              (FVectorVMState *VectorVMState, int InputRegisterIdx, int DataSetIdx, bool IntReg, uint16 *OptOutIdxInRegTable);
 VECTORVM_API uint16           VVMDbgGetRemappedOutputIdx             (FVectorVMState *VectorVMState, int OutputRegisterIdx, int DataSetIdx, bool IntReg, uint16 *OptOutIdxInRegTable);
 VECTORVM_API uint16           VVMDbgGetRemappedConstIdx              (FVectorVMState *VectorVMState, int ConstTableIdx, int ConstBuffIdx, uint16 *OptOutIdxInRegTable);
-
-#else // VECTORVM_SUPPORTS_EXPERIMENTAL
-
-struct FVectorVMState {
-
-};
-
-VECTORVM_API void    SerializeVectorVMWriteToFile    (FVectorVMSerializeState *SerializeState, uint8 WhichStateWritten, const wchar_t *Filename);
-VECTORVM_API void    FreeVectorVMSerializeState      (FVectorVMSerializeState *SerializeState);
-
-#endif // VECTORVM_SUPPORTS_EXPERIMENTAL
-
-//things needed for when both experimental and legacy are built
-#if VECTORVM_SUPPORTS_LEGACY
-VECTORVM_API uint32  SerializeVectorVMInputDataSets  (FVectorVMSerializeState *SerializeState, TArrayView<FDataSetMeta>, const uint8 * const *ConstantTableData, const int *ConstantTableSizes, int32 ConstantTableCount); //only use when not calling AllocVectorVMState()
-VECTORVM_API uint32  SerializeVectorVMOutputDataSets (FVectorVMSerializeState *SerializeState, TArrayView<FDataSetMeta>, const uint8 * const *ConstantTableData, const int *ConstantTableSizes, int32 ConstantTableCount);
-#endif
\ No newline at end of file
diff --git a/Engine/Source/Runtime/VectorVM/Public/VectorVMLegacy.h b/Engine/Source/Runtime/VectorVM/Public/VectorVMLegacy.h
deleted file mode 100644
diff --git a/Engine/Source/Runtime/VectorVM/VectorVM.Build.cs b/Engine/Source/Runtime/VectorVM/VectorVM.Build.cs
--- a/Engine/Source/Runtime/VectorVM/VectorVM.Build.cs
+++ b/Engine/Source/Runtime/VectorVM/VectorVM.Build.cs
@@ -16,8 +16,6 @@ public class VectorVM : ModuleRules
         PublicDefinitions.AddRange(
             new string[]
             {
-                "VECTORVM_SUPPORTS_EXPERIMENTAL=1",
-                "VECTORVM_SUPPORTS_LEGACY=1",
                 "VECTORVM_SUPPORTS_SERIALIZATION=0",
                 "VECTORVM_DEBUG_PRINTF=0"
             });
