Cleanup VectorVM pt 1. Remove VECTORVM_SUPPORTS_SERIALIZATION path
#tests EngineTests

diff --git a/Engine/Plugins/FX/Niagara/Source/Niagara/Private/NiagaraScriptExecutionContext.cpp b/Engine/Plugins/FX/Niagara/Source/Niagara/Private/NiagaraScriptExecutionContext.cpp
--- a/Engine/Plugins/FX/Niagara/Source/Niagara/Private/NiagaraScriptExecutionContext.cpp
+++ b/Engine/Plugins/FX/Niagara/Source/Niagara/Private/NiagaraScriptExecutionContext.cpp
@@ -251,7 +251,7 @@ bool FNiagaraScriptExecutionContextBase::ExecuteInternal(uint32 NumInstances, co
 
 	if (VectorVMState)
 	{
-		ExecVectorVMState(&ExecCtx, nullptr, nullptr);
+		ExecVectorVMState(&ExecCtx);
 	}
 
 #if STATS
diff --git a/Engine/Source/Runtime/VectorVM/Private/VectorVM.cpp b/Engine/Source/Runtime/VectorVM/Private/VectorVM.cpp
--- a/Engine/Source/Runtime/VectorVM/Private/VectorVM.cpp
+++ b/Engine/Source/Runtime/VectorVM/Private/VectorVM.cpp
@@ -1,7 +1,6 @@
 // Copyright Epic Games, Inc. All Rights Reserved.
 
 #include "VectorVM.h"
-#include "VectorVMSerialization.h"
 #include "Modules/ModuleManager.h"
 #include "UObject/Class.h"
 #include "UObject/Package.h"
@@ -261,7 +260,7 @@ UEnum* g_VectorVMEnumOperandObj = nullptr;
 static const char *VVM_OP_NAMES[] {
 	VVM_OP_XM_LIST
 };
-#undef VVM_OP_XM
+#undef VVM_OP_XMd 
 
 FString VectorVM::GetOpName(EVectorVMOp Op)
 {
diff --git a/Engine/Source/Runtime/VectorVM/Private/VectorVMExperimental.cpp b/Engine/Source/Runtime/VectorVM/Private/VectorVMExperimental.cpp
--- a/Engine/Source/Runtime/VectorVM/Private/VectorVMExperimental.cpp
+++ b/Engine/Source/Runtime/VectorVM/Private/VectorVMExperimental.cpp
@@ -4,18 +4,6 @@
 #include "VectorVM.h"
 #include "HAL/ConsoleManager.h"
 
-//prototypes for internal functions to the VVM
-#if VECTORVM_SUPPORTS_SERIALIZATION
-void VVMSer_serializeInstruction(FVectorVMSerializeState *SerializeState, FVectorVMSerializeState *CmpSerializeState, 
-                                 FVectorVMState *VVMState, FVectorVMBatchState *BatchState, 
-                                 int StartInstance, int NumInstancesThisChunk, int NumLoops, 
-                                 int OpStart, int NumOps, uint64 StartInsCycles, uint64 EndInsCycles);
-
-uint32 VVMSer_initSerializationState_(FVectorVMSerializeState *SerializeState, FVectorVMExecContext *ExecCtx, const FVectorVMOptimizeContext *OptimizeContext, uint32 Flags);
-int VVMGetRegisterType(FVectorVMState *VVMState, uint16 RegIdx, uint16 *OutAbsReg);
-uint32 *VVMSer_getRegPtrTablePtrFromIns(FVectorVMSerializeState *SerializeState, FVectorVMSerializeInstruction *Ins, uint16 RegIdx);
-#endif
-
 #if VECTORVM_DEBUG_PRINTF
 
 #define VECTORVM_PRINTF(fmt, ...)	                    \
@@ -1175,88 +1163,6 @@ VM_FORCEINLINE VectorRegister4f VVM_nextRandom(FVectorVMBatchState *BatchState,
 	return VectorMultiply(VectorSubtract(VectorRegister4f(VectorCastIntToFloat(VectorIntOr(VectorShiftRightImmLogical(VVMXorwowStep(BatchState), 9), VectorIntSet1(0x3F800000)))), VVM_m128Const(One)), a);
 }
 
-#if VECTORVM_SUPPORTS_SERIALIZATION
-bool VVM_serSyncRandom(const uint8 *InsPtr, FVectorVMBatchState *BatchState, FVectorVMExecContext *ExecCtx, FVectorVMSerializeState *SerializeState, FVectorVMSerializeState *CmpSerializeState, int OutputRegPtrIdx, EVectorVMOp RandomOpCode)
-{
-	if (SerializeState && (SerializeState->Flags & VVMSer_SyncRandom) && CmpSerializeState && CmpSerializeState->Flags & VVMSer_OptimizedBytecode)
-	{
-		int InsByteCodeOffset = (int)(SerializeState->Running.StartOpPtr - ExecCtx->VVMState->Bytecode);
-		if (CmpSerializeState->NumInstructions > SerializeState->NumInstructions && 
-			CmpSerializeState->Instructions[SerializeState->NumInstructions].OpStart == InsByteCodeOffset && 
-			(EVectorVMOp)CmpSerializeState->Bytecode[InsByteCodeOffset - 1] == RandomOpCode)
-		{
-			uint16 RegIdx = ((uint16 *)InsPtr)[OutputRegPtrIdx];
-			FVectorVMSerializeInstruction *CmpIns = CmpSerializeState->Instructions + SerializeState->NumInstructions;
-			float *SrcReg  = (float *)VVMSer_getRegPtrTablePtrFromIns(CmpSerializeState, CmpIns, ((uint16 *)(CmpSerializeState->Bytecode + InsByteCodeOffset))[OutputRegPtrIdx]);
-			float * DstReg = (float *)BatchState->RegPtrTable[RegIdx];
-			int DstInc     = (int)BatchState->RegIncTable[RegIdx];
-			if (DstInc != 0)
-			{
-				VVMMemCpy(DstReg, SrcReg, sizeof(float) * BatchState->ChunkLocalData.NumInstancesThisChunk);
-			}
-		}
-		return true;
-	}
-	else if (SerializeState && (SerializeState->Flags & VVMSer_SyncRandom) && CmpSerializeState && SerializeState->NumInstances == CmpSerializeState->NumInstances && (CmpSerializeState->NumInstructions > SerializeState->NumInstructions))
-	{
-		FVectorVMSerializeInstruction *CmpIns = nullptr; //this instruction
-		FVectorVMOptimizeInstruction *OptIns = nullptr;
-		for (uint32 i = 0 ; i < SerializeState->OptimizeCtx->Intermediate.NumInstructions; ++i)
-		{
-			if (SerializeState->OptimizeCtx->Intermediate.Instructions[i].OpCode == RandomOpCode && SerializeState->OptimizeCtx->Intermediate.Instructions[i].PtrOffsetInOptimizedBytecode == (int)(SerializeState->Running.StartOpPtr - ExecCtx->VVMState->Bytecode))
-			{
-				OptIns = SerializeState->OptimizeCtx->Intermediate.Instructions + i;
-				break;
-			}
-		}
-		if (OptIns)
-		{
-			for (uint32 i = 0; i < CmpSerializeState->NumInstructions; ++i)
-			{
-				if (CmpSerializeState->Bytecode[CmpSerializeState->Instructions[i].OpStart] == (uint8)RandomOpCode && OptIns->PtrOffsetInOrigBytecode == CmpSerializeState->Instructions[i].OpStart)
-				{
-					CmpIns = CmpSerializeState->Instructions + i;
-					break;
-				}
-			}
-		}
-		check(false);
-		if (OptIns && CmpIns)
-		{
-			//for (int i = 0; i < NumRegisters; ++i)
-			//{
-			//	uint8 RegIdx = InsPtr[1 + i];
-			//	if (RegIdx != 0xFF) {
-			//		uint32 * DstReg  = (uint32 *)BatchState->RegPtrTable[RegIdx];
-			//		int DstInc       = (int)BatchState->RegIncTable[RegIdx];
-			//		if (DstInc != 0) { //skip constants
-			//			int SrcRegIdx = 0xFFFF;
-			//			if (OptIns && (SerializeState->Flags & VVMSer_OptimizedBytecode) && !(CmpSerializeState->Flags & VVMSer_OptimizedBytecode))
-			//			{
-			//				SrcRegIdx = ((uint16 *)(CmpSerializeState->Bytecode + OptIns->PtrOffsetInOrigBytecode + 2))[i] & 0x7FFF; //high bit is register in original bytecode
-			//			}
-			//			if (SrcRegIdx != 0x7FFF) //invalid register, skipped by external function in the execution
-			//			{
-			//				//can only be temp registers because there's no output aliasing in the original VM
-			//				uint8 *SrcReg = (uint8 *)(CmpIns->RegisterTable + SrcRegIdx * CmpSerializeState->NumInstances + StartInstanceThisChunk);
-			//				VVMMemCpy(DstReg, SrcReg, sizeof(uint32) * NumInstancesThisChunk);
-			//			}
-			//		}
-			//	}
-			//}
-		}
-		return true;
-	}
-	return false;
-
-}
-#else //VECTORVM_SUPPORTS_SERIALIZATION
-VM_FORCEINLINE bool VVM_serSyncRandom(const uint8 *InsPtr, FVectorVMBatchState *BatchState, FVectorVMExecContext *ExecCtx, FVectorVMSerializeState *SerializeState, FVectorVMSerializeState *CmpSerializeState, int OutputRegPtrIdx, EVectorVMOp RandomOpCode)
-{
-	return false;
-}
-#endif //VECTORVM_SUPPORTS_SERIALIZATION
-
 static VM_FORCEINLINE int32 VVM_SafeIntDivide(int32 Numerator, int32 Denominator)
 {
 	static constexpr int32 MinIntValue = std::numeric_limits<int32>::min();
@@ -1289,139 +1195,126 @@ static VM_FORCEINLINE VectorRegister4i VVMIntDiv(VectorRegister4i v0, VectorRegi
 }
 
 
-static VM_FORCEINLINE const uint8 *VVM_random(const bool CT_MultipleLoops, const uint8 *InsPtr, FVectorVMBatchState *BatchState, FVectorVMExecContext *ExecCtx, FVectorVMSerializeState *SS, FVectorVMSerializeState *CmpSS, int NumLoops)
+static VM_FORCEINLINE const uint8 *VVM_random(const bool CT_MultipleLoops, const uint8 *InsPtr, FVectorVMBatchState *BatchState, FVectorVMExecContext *ExecCtx, int NumLoops)
 {
-	VVMSer_instruction(0, 1);
-	if (!VVM_serSyncRandom(InsPtr, BatchState, ExecCtx, SS, CmpSS, 1, EVectorVMOp::random))
+	uint16* RegIndices = (uint16*)InsPtr;
+	uint8 *P0   = BatchState->RegPtrTable[RegIndices[0]];
+	uint8 *P1   = BatchState->RegPtrTable[RegIndices[1]];
+	uint32 Inc0 = (uint32)BatchState->RegIncTable[RegIndices[0]];
+	if (CT_MultipleLoops)
 	{
-		uint16* RegIndices = (uint16*)InsPtr;
-		uint8 *P0   = BatchState->RegPtrTable[RegIndices[0]];
-		uint8 *P1   = BatchState->RegPtrTable[RegIndices[1]];
-		uint32 Inc0 = (uint32)BatchState->RegIncTable[RegIndices[0]];
-		if (CT_MultipleLoops)
-		{
-			uint8 *End = P1 + sizeof(FVecReg) * NumLoops;
-			do
-			{
-				VectorRegister4f R0 = VectorLoad((float *)P0);
-				P0 += Inc0;
-				VectorRegister4f Res = VVM_nextRandom(BatchState, R0);
-				VectorStore(Res, (float *)P1);
-				P1 += sizeof(FVecReg);
-			} while (P1 < End);
-		}
-		else
+		uint8 *End = P1 + sizeof(FVecReg) * NumLoops;
+		do
 		{
 			VectorRegister4f R0 = VectorLoad((float *)P0);
+			P0 += Inc0;
 			VectorRegister4f Res = VVM_nextRandom(BatchState, R0);
 			VectorStore(Res, (float *)P1);
-		}
+			P1 += sizeof(FVecReg);
+		} while (P1 < End);
 	}
+	else
+	{
+		VectorRegister4f R0 = VectorLoad((float *)P0);
+		VectorRegister4f Res = VVM_nextRandom(BatchState, R0);
+		VectorStore(Res, (float *)P1);
+	}
+
 	return InsPtr + 5;
 }
 
-static VM_FORCEINLINE const uint8 *VVM_randomi(const bool CT_MultipleLoops, const uint8 *InsPtr, FVectorVMBatchState *BatchState, FVectorVMExecContext *ExecCtx, FVectorVMSerializeState *SS, FVectorVMSerializeState *CmpSS, int NumLoops)
+static VM_FORCEINLINE const uint8 *VVM_randomi(const bool CT_MultipleLoops, const uint8 *InsPtr, FVectorVMBatchState *BatchState, FVectorVMExecContext *ExecCtx, int NumLoops)
 {
-	VVMSer_instruction(1, 1);
-	if (!VVM_serSyncRandom(InsPtr, BatchState, ExecCtx, SS, CmpSS, 1, EVectorVMOp::randomi))
+	uint16* RegIndices = (uint16*)InsPtr;
+	uint8 *P0   = BatchState->RegPtrTable[RegIndices[0]];
+	uint8 *P1   = BatchState->RegPtrTable[RegIndices[1]];
+	uint32 Inc0 = (uint32)BatchState->RegIncTable[RegIndices[0]];
+	if (CT_MultipleLoops)
 	{
-		uint16* RegIndices = (uint16*)InsPtr;
-		uint8 *P0   = BatchState->RegPtrTable[RegIndices[0]];
-		uint8 *P1   = BatchState->RegPtrTable[RegIndices[1]];
-		uint32 Inc0 = (uint32)BatchState->RegIncTable[RegIndices[0]];
-		if (CT_MultipleLoops)
-		{
-			uint8 *End = P1 + sizeof(FVecReg) * NumLoops;
-			do
-			{
-				VectorRegister4f R0 = VectorLoad((float *)P0);
-				P0 += Inc0;
-				VectorRegister4f Res = VVM_nextRandom(BatchState, R0);
-				VectorIntStore(VectorFloatToInt(Res), P0);
-				P1 += sizeof(FVecReg);
-			} while (P1 < End);
-		}
-		else
+		uint8 *End = P1 + sizeof(FVecReg) * NumLoops;
+		do
 		{
 			VectorRegister4f R0 = VectorLoad((float *)P0);
+			P0 += Inc0;
 			VectorRegister4f Res = VVM_nextRandom(BatchState, R0);
 			VectorIntStore(VectorFloatToInt(Res), P0);
-		}
-}
+			P1 += sizeof(FVecReg);
+		} while (P1 < End);
+	}
+	else
+	{
+		VectorRegister4f R0 = VectorLoad((float *)P0);
+		VectorRegister4f Res = VVM_nextRandom(BatchState, R0);
+		VectorIntStore(VectorFloatToInt(Res), P0);
+	}
 	return InsPtr + 5;
 }
 
-static VM_FORCEINLINE const uint8 *VVM_random_add(const bool CT_MultipleLoops, const uint8 *InsPtr, FVectorVMBatchState *BatchState, FVectorVMExecContext *ExecCtx, FVectorVMSerializeState *SS, FVectorVMSerializeState *CmpSS, int NumLoops)
+static VM_FORCEINLINE const uint8 *VVM_random_add(const bool CT_MultipleLoops, const uint8 *InsPtr, FVectorVMBatchState *BatchState, FVectorVMExecContext *ExecCtx, int NumLoops)
 {
-	VVMSer_instruction(0, 2);
-	if (!VVM_serSyncRandom(InsPtr, BatchState, ExecCtx, SS, CmpSS, 2, EVectorVMOp::random_add))
-	{
-		uint16* RegIndices = (uint16*)InsPtr;
-		uint8 *P0   = BatchState->RegPtrTable[RegIndices[0]];
-		uint8 *P1   = BatchState->RegPtrTable[RegIndices[1]];
-		uint8 *P2   = BatchState->RegPtrTable[RegIndices[2]];
-		uint32 Inc0 = (uint32)BatchState->RegIncTable[RegIndices[0]];
-		uint32 Inc1 = (uint32)BatchState->RegIncTable[RegIndices[1]];
+	uint16* RegIndices = (uint16*)InsPtr;
+	uint8 *P0   = BatchState->RegPtrTable[RegIndices[0]];
+	uint8 *P1   = BatchState->RegPtrTable[RegIndices[1]];
+	uint8 *P2   = BatchState->RegPtrTable[RegIndices[2]];
+	uint32 Inc0 = (uint32)BatchState->RegIncTable[RegIndices[0]];
+	uint32 Inc1 = (uint32)BatchState->RegIncTable[RegIndices[1]];
 
-		if (CT_MultipleLoops)
-		{
-			uint8 *End = P2 + sizeof(FVecReg) * NumLoops;
-			do
-			{
-				VectorRegister4f R0 = VectorLoad((float *)P0);
-				VectorRegister4f R1 = VectorLoad((float *)P1);
-				P0 += Inc0;
-				P1 += Inc1;
-				VectorRegister4f Res = VectorAdd(VVM_nextRandom(BatchState, R0), R1);
-				VectorStore(Res, (float *)P2);
-				P2 += sizeof(FVecReg);
-			} while (P2 < End);
-		}
-		else
+	if (CT_MultipleLoops)
+	{
+		uint8 *End = P2 + sizeof(FVecReg) * NumLoops;
+		do
 		{
 			VectorRegister4f R0 = VectorLoad((float *)P0);
 			VectorRegister4f R1 = VectorLoad((float *)P1);
+			P0 += Inc0;
+			P1 += Inc1;
 			VectorRegister4f Res = VectorAdd(VVM_nextRandom(BatchState, R0), R1);
 			VectorStore(Res, (float *)P2);
-		}
+			P2 += sizeof(FVecReg);
+		} while (P2 < End);
 	}
+	else
+	{
+		VectorRegister4f R0 = VectorLoad((float *)P0);
+		VectorRegister4f R1 = VectorLoad((float *)P1);
+		VectorRegister4f Res = VectorAdd(VVM_nextRandom(BatchState, R0), R1);
+		VectorStore(Res, (float *)P2);
+	}
+
 	return InsPtr + 7;
 }
 
-static VM_FORCEINLINE const uint8 *VVM_random_2x(const bool CT_MultipleLoops, const uint8 *InsPtr, FVectorVMBatchState *BatchState, FVectorVMExecContext *ExecCtx, FVectorVMSerializeState *SS, FVectorVMSerializeState *CmpSS, int NumLoops)
+static VM_FORCEINLINE const uint8 *VVM_random_2x(const bool CT_MultipleLoops, const uint8 *InsPtr, FVectorVMBatchState *BatchState, FVectorVMExecContext *ExecCtx, int NumLoops)
 {
-	VVMSer_instruction(0, 2);
-	if (!VVM_serSyncRandom(InsPtr, BatchState, ExecCtx, SS, CmpSS, 1, EVectorVMOp::random_2x) && !VVM_serSyncRandom(InsPtr, BatchState, ExecCtx, SS, CmpSS, 2, EVectorVMOp::random_2x))
+	uint16* RegIndices = (uint16*)InsPtr;
+	uint8* P0 = BatchState->RegPtrTable[RegIndices[0]];
+	uint8* P1 = BatchState->RegPtrTable[RegIndices[1]];
+	uint8* P2 = BatchState->RegPtrTable[RegIndices[2]];
+	uint32 Inc0 = (uint32)BatchState->RegIncTable[RegIndices[0]];
+	if (CT_MultipleLoops)
 	{
-		uint16* RegIndices = (uint16*)InsPtr;
-		uint8 *P0   = BatchState->RegPtrTable[RegIndices[0]];
-		uint8 *P1   = BatchState->RegPtrTable[RegIndices[1]];
-		uint8 *P2   = BatchState->RegPtrTable[RegIndices[2]];
-		uint32 Inc0 = (uint32)BatchState->RegIncTable[RegIndices[0]];
-		if (CT_MultipleLoops)
-		{
-			uint8 *End = P1 + sizeof(FVecReg) * NumLoops;
-			do
-			{
-				VectorRegister4f R0 = VectorLoad((float *)P0);
-				P0 += Inc0;
-				VectorRegister4f Res0 = VVM_nextRandom(BatchState, R0);
-				VectorRegister4f Res1 = VVM_nextRandom(BatchState, R0);
-				VectorStore(Res0, (float *)P1);
-				VectorStore(Res1, (float *)P2);
-				P1 += sizeof(FVecReg);
-				P2 += sizeof(FVecReg);
-			} while (P1 < End);
-		}
-		else
+		uint8* End = P1 + sizeof(FVecReg) * NumLoops;
+		do
 		{
-			VectorRegister4f R0 = VectorLoad((float *)P0);
+			VectorRegister4f R0 = VectorLoad((float*)P0);
+			P0 += Inc0;
 			VectorRegister4f Res0 = VVM_nextRandom(BatchState, R0);
 			VectorRegister4f Res1 = VVM_nextRandom(BatchState, R0);
-			VectorStore(Res0, (float *)P1);
-			VectorStore(Res1, (float *)P2);
-		}
+			VectorStore(Res0, (float*)P1);
+			VectorStore(Res1, (float*)P2);
+			P1 += sizeof(FVecReg);
+			P2 += sizeof(FVecReg);
+		} while (P1 < End);
 	}
+	else
+	{
+		VectorRegister4f R0 = VectorLoad((float*)P0);
+		VectorRegister4f Res0 = VVM_nextRandom(BatchState, R0);
+		VectorRegister4f Res1 = VVM_nextRandom(BatchState, R0);
+		VectorStore(Res0, (float*)P1);
+		VectorStore(Res1, (float*)P2);
+	}
+
 	return InsPtr + 7;
 }
 
@@ -1708,7 +1601,6 @@ static const uint8 *VVM_update_id(const bool CT_MultipleLoops, const uint8 *InsP
 
 static const uint8 *VVM_acquire_id(const bool CT_MultipleLoops, const uint8 *InsPtr, FVectorVMBatchState *BatchState, FVectorVMExecContext *ExecCtx, int NumLoops)
 {
-	VVMSer_instruction(1, 2);
 	uint8 DataSetIdx = InsPtr[4];
 	check(DataSetIdx < (uint32)ExecCtx->DataSets.Num());
 	FDataSetMeta *DataSet = &ExecCtx->DataSets[DataSetIdx];
@@ -1745,7 +1637,7 @@ static const uint8 *VVM_acquire_id(const bool CT_MultipleLoops, const uint8 *Ins
 	return InsPtr + 6;
 }
 
-const uint8 *VVM_external_func_call(const bool CT_MultipleLoops, const uint8 *InsPtr, FVectorVMBatchState *BatchState, FVectorVMExecContext *ExecCtx, FVectorVMSerializeState *SerializeState, FVectorVMSerializeState *CmpSerializeState, int NumLoops)
+const uint8 *VVM_external_func_call(const bool CT_MultipleLoops, const uint8 *InsPtr, FVectorVMBatchState *BatchState, FVectorVMExecContext *ExecCtx, int NumLoops)
 {
 	int FnIdx = (int)*(uint16 *)(InsPtr);
 	FVectorVMExtFunctionData *ExtFnData = ExecCtx->VVMState->ExtFunctionTable + FnIdx;
@@ -1757,142 +1649,9 @@ const uint8 *VVM_external_func_call(const bool CT_MultipleLoops, const uint8 *In
 		VECTORVM_PRINTF("%c%d, ", VVM_RT_CHAR[RegType], AbsRegIdx);
 	}
 	VECTORVM_PRINT_NEW_LINE();
-	
 #endif //VECTORVM_DEBUG_PRINTF
-#if	VECTORVM_SUPPORTS_SERIALIZATION && !defined(VVM_SERIALIZE_NO_WRITE)
-	if (SerializeState && (SerializeState->Flags & VVMSer_SyncExtFns) && CmpSerializeState && SerializeState->NumInstances == CmpSerializeState->NumInstances && (CmpSerializeState->NumInstructions > SerializeState->NumInstructions))
-	{
-		//If we hit this branch we are using the output from the comparision state instead of running the external function itself.
-		//AFAIK the VM is not speced to have the inputs and outputs in a particular order, and even if it is we shouldn't rely on
-		//3rd party external function writers to follow the spec.  Therefore we don't just sync what we think is output, we sync
-		//all temp registers that are used in the function.
-		int NumRegisters = ExtFnData->NumInputs + ExtFnData->NumOutputs;
-		FVectorVMSerializeInstruction *CmpIns = nullptr; //this instruction
-		FVectorVMOptimizeInstruction *OptIns = nullptr;
-		if (SerializeState->OptimizeCtx && SerializeState->OptimizeCtx->Intermediate.NumInstructions)
-		{
-			//instructions have been re-ordered, can't binary search, must linear search
-			for (uint32 i = 0 ; i < SerializeState->OptimizeCtx->Intermediate.NumInstructions; ++i)
-			{
-				if (SerializeState->OptimizeCtx->Intermediate.Instructions[i].OpCode == EVectorVMOp::external_func_call && SerializeState->OptimizeCtx->Intermediate.Instructions[i].PtrOffsetInOptimizedBytecode == (int)(SerializeState->Running.StartOpPtr - ExecCtx->VVMState->Bytecode))
-				{
-					OptIns = SerializeState->OptimizeCtx->Intermediate.Instructions + i;
-					break;
-				}
-			}
-			if (OptIns)
-			{
-				for (uint32 i = 0; i < CmpSerializeState->NumInstructions; ++i)
-				{
-					if (CmpSerializeState->Bytecode[CmpSerializeState->Instructions[i].OpStart] == (uint8)EVectorVMOp::external_func_call && OptIns->PtrOffsetInOrigBytecode == CmpSerializeState->Instructions[i].OpStart)
-					{
-						CmpIns = CmpSerializeState->Instructions + i;
-						break;
-					}
-				}
-			}
-			if (OptIns && CmpIns)
-			{
-				for (int i = 0; i < NumRegisters; ++i)
-				{
-					uint8 RegIdx = InsPtr[1 + i];
-					if (RegIdx != 0xFF) {
-						uint32 * DstReg  = (uint32 *)BatchState->RegPtrTable[RegIdx];
-						int DstInc       = (int)BatchState->RegIncTable[RegIdx];
-						if (DstInc != 0) { //skip constants
-							int SrcRegIdx = 0xFFFF;
-							if (OptIns && (SerializeState->Flags & VVMSer_OptimizedBytecode) && !(CmpSerializeState->Flags & VVMSer_OptimizedBytecode))
-							{
-								SrcRegIdx = ((uint16 *)(CmpSerializeState->Bytecode + OptIns->PtrOffsetInOrigBytecode + 2))[i] & 0x7FFF; //high bit is register in original bytecode
-							}
-							if (SrcRegIdx != 0x7FFF) //invalid register, skipped by external function in the execution
-							{
-								//can only be temp registers because there's no output aliasing in the original VM
-								uint8 *SrcReg = (uint8 *)(CmpIns->RegisterTable + SrcRegIdx * CmpSerializeState->NumInstances + BatchState->ChunkLocalData.StartInstanceThisChunk);
-								VVMMemCpy(DstReg, SrcReg, sizeof(uint32) * BatchState->ChunkLocalData.NumInstancesThisChunk);
-							}
-						}
-					}
-				}
-			}
-		}
-		else if ((SerializeState->Flags & VVMSer_OptimizedBytecode) && (CmpSerializeState->Flags & VVMSer_OptimizedBytecode))
-		{
-			//we do not have the optimized internal state (reading from serialized data from a cooked build... therefore we're starting from an optimized script and comparing to a (different?) optimized
-			//script, and thus the instructions should line up
-			//CmpIns = CmpSerializeState->Instructions + SerializeState->Running.NumInstructionsThisChunk;
-			int ExtFnInsIdx = 0;
-			for (int i = 0; i < SerializeState->Running.NumInstructionsThisChunk; ++i)
-			{
-				int OpCodeIdx = SerializeState->Instructions[i].OpStart - 1;
-				if (OpCodeIdx >= 0 && OpCodeIdx < (int)SerializeState->NumBytecodeBytes)
-				{
-					EVectorVMOp InsOpCode = (EVectorVMOp)SerializeState->Bytecode[OpCodeIdx];
-					if (InsOpCode == EVectorVMOp::external_func_call)
-					{
-						++ExtFnInsIdx;
-					}
-				}
-			}
-			//now find the instruction in the comparison serialize state
-			int SearchExtFnInsIdx = 0;
-			for (uint32 i = 0; i < CmpSerializeState->NumInstructions; ++i)
-			{
-				int OpCodeIdx = CmpSerializeState->Instructions[i].OpStart - 1;
-				if (OpCodeIdx >= 0 && OpCodeIdx < (int)CmpSerializeState->NumBytecodeBytes)
-				{
-					EVectorVMOp InsOpCode = (EVectorVMOp)CmpSerializeState->Bytecode[OpCodeIdx];
-					if (InsOpCode == EVectorVMOp::external_func_call)
-					{
-						if (SearchExtFnInsIdx == ExtFnInsIdx) {
-							CmpIns = CmpSerializeState->Instructions +i;
-							break;
-						}
-						++SearchExtFnInsIdx;
-					}
-				}
-			}
-					
-			if (CmpIns)
-			{
-				uint16 *SrcRegIndices = (uint16 *)(CmpSerializeState->Bytecode + CmpIns->OpStart + sizeof(uint16));
-				uint16 *DstRegIndices = ((uint16 *)InsPtr) + 1;
-				for (int i = 0; i < NumRegisters; ++i)
-				{
-					uint16 SrcRegIdx = SrcRegIndices[i];
-					uint16 DstRegIdx = DstRegIndices[i];
 
-					if (SrcRegIdx != 0xFFFF && DstRegIdx != 0xFFFF)
-					{
-						uint8 DstInc   = BatchState->RegIncTable[DstRegIdx];
-						if (DstInc != 0) //skip constants
-						{
-							int SrcRegType = VVMGetRegisterType(ExecCtx->VVMState, SrcRegIdx, &SrcRegIdx);
-							uint32 *DstReg = (uint32 *)BatchState->RegPtrTable[DstRegIdx];
-							uint32 *SrcReg = NULL;
-							if (SrcRegType == VVM_RT_TEMPREG)
-							{
-								SrcReg = (uint32 *)CmpIns->RegisterTable + SrcRegIdx * CmpSerializeState->NumInstances + BatchState->ChunkLocalData.StartInstanceThisChunk;
-							}
-							else if (SrcRegType == VVM_RT_OUTPUT)
-							{
-								SrcReg = (uint32 *)CmpIns->RegisterTable + (CmpSerializeState->NumTempRegisters + SrcRegIdx) * CmpSerializeState->NumInstances + BatchState->ChunkLocalData.StartInstanceThisChunk;
-							}
-							if (SrcReg)
-							{
-								VVMMemCpy(DstReg, SrcReg, sizeof(uint32) * BatchState->ChunkLocalData.NumInstancesThisChunk);
-							}
-						}
-					}
-				}
-			}
-		}		
-	}
-	else
-	{
-#	else //VECTORVM_SUPPORTS_SERIALIZATION
 	{
-#	endif //VECTORVM_SUPPORTS_SERIALIZATION
 		check(*InsPtr < ExecCtx->VVMState->NumExtFunctions);
 		check((uint32)(ExtFnData->NumInputs + ExtFnData->NumOutputs) <= ExecCtx->VVMState->MaxExtFnRegisters);
 		const uint16 *RegIndices      = ((uint16 *)InsPtr) + 1;
@@ -1964,7 +1723,7 @@ static VM_FORCEINLINE const uint8 *VVM_sincos(const bool CT_MultipleLoops, const
 
 //zF16 - stb_image_resize
 
-#define VVM_NULL_FN_ARGS const bool CT_MultipleLoops, const uint8 *InsPtr, FVectorVMBatchState *BatchState, FVectorVMExecContext *ExecCtx, FVectorVMSerializeState *SerializeState, FVectorVMSerializeState *CmpSerializeState, int NumLoops
+#define VVM_NULL_FN_ARGS const bool CT_MultipleLoops, const uint8 *InsPtr, FVectorVMBatchState *BatchState, FVectorVMExecContext *ExecCtx, int NumLoops
 
 typedef const uint8 *    (VVMFn_null)(VVM_NULL_FN_ARGS);
 typedef VectorRegister4f (VVMFn_1f)  (FVectorVMBatchState *, VectorRegister4f a);
@@ -1977,33 +1736,32 @@ typedef VectorRegister4i (VVMFn_2i)  (FVectorVMBatchState *, VectorRegister4i a,
 typedef VectorRegister4i (VVMFn_3i)  (FVectorVMBatchState *, VectorRegister4i a, VectorRegister4i b, VectorRegister4i c);
 
 //Dispatch functions: these rely on VM_FORCEINLINE to actually force the function to be inlined
-VM_FORCEINLINE const uint8 *VVM_Dispatch_execFn0null_0null(const bool CT_MultipleLoops, const uint8 *InsPtr, FVectorVMBatchState *BatchState, FVectorVMExecContext *ExecCtx, FVectorVMSerializeState *SerializeState, FVectorVMSerializeState *CmpSerializeState, VVMFn_null fn, int NumLoops) {
-	VVMSer_instruction(0, 0);
-	return fn(CT_MultipleLoops, InsPtr, BatchState, ExecCtx, SerializeState, CmpSerializeState, NumLoops);
+VM_FORCEINLINE const uint8 *VVM_Dispatch_execFn0null_0null(const bool CT_MultipleLoops, const uint8 *InsPtr, FVectorVMBatchState *BatchState, FVectorVMExecContext *ExecCtx, VVMFn_null fn, int NumLoops)
+{
+	return fn(CT_MultipleLoops, InsPtr, BatchState, ExecCtx, NumLoops);
 }
 
-VM_FORCEINLINE const uint8 *VVM_Dispatch_execFn1null_0null(const bool CT_MultipleLoops, const uint8 *InsPtr, FVectorVMBatchState *BatchState, FVectorVMExecContext *ExecCtx, FVectorVMSerializeState *SerializeState, FVectorVMSerializeState *CmpSerializeState, VVMFn_null fn, int NumLoops) {
-	VVMSer_instruction(0, 0);
-	return fn(CT_MultipleLoops, InsPtr, BatchState, ExecCtx, SerializeState, CmpSerializeState, NumLoops);
+VM_FORCEINLINE const uint8 *VVM_Dispatch_execFn1null_0null(const bool CT_MultipleLoops, const uint8 *InsPtr, FVectorVMBatchState *BatchState, FVectorVMExecContext *ExecCtx, VVMFn_null fn, int NumLoops)
+{
+	return fn(CT_MultipleLoops, InsPtr, BatchState, ExecCtx, NumLoops);
 }
 
-VM_FORCEINLINE const uint8 *VVM_Dispatch_execFn1null_1null(const bool CT_MultipleLoops, const uint8 *InsPtr, FVectorVMBatchState *BatchState, FVectorVMExecContext *ExecCtx, FVectorVMSerializeState *SerializeState, FVectorVMSerializeState *CmpSerializeState, VVMFn_null fn, int NumLoops) {
-	VVMSer_instruction(0, 1);
-	return fn(CT_MultipleLoops, InsPtr, BatchState, ExecCtx, SerializeState, CmpSerializeState, NumLoops);
+VM_FORCEINLINE const uint8 *VVM_Dispatch_execFn1null_1null(const bool CT_MultipleLoops, const uint8 *InsPtr, FVectorVMBatchState *BatchState, FVectorVMExecContext *ExecCtx, VVMFn_null fn, int NumLoops)
+{
+	return fn(CT_MultipleLoops, InsPtr, BatchState, ExecCtx, NumLoops);
 }
 
-VM_FORCEINLINE const uint8 *VVM_Dispatch_execFn1null_2null(const bool CT_MultipleLoops, const uint8 *InsPtr, FVectorVMBatchState *BatchState, FVectorVMExecContext *ExecCtx, FVectorVMSerializeState *SerializeState, FVectorVMSerializeState *CmpSerializeState, VVMFn_null fn, int NumLoops) {
-	VVMSer_instruction(0, 2);
-	return fn(CT_MultipleLoops, InsPtr, BatchState, ExecCtx, SerializeState, CmpSerializeState, NumLoops);
+VM_FORCEINLINE const uint8 *VVM_Dispatch_execFn1null_2null(const bool CT_MultipleLoops, const uint8 *InsPtr, FVectorVMBatchState *BatchState, FVectorVMExecContext *ExecCtx, VVMFn_null fn, int NumLoops)
+{
+	return fn(CT_MultipleLoops, InsPtr, BatchState, ExecCtx, NumLoops);
 }
 
-VM_FORCEINLINE const uint8 *VVM_Dispatch_execFn2null_1null(const bool CT_MultipleLoops, const uint8 *InsPtr, FVectorVMBatchState *BatchState, FVectorVMExecContext *ExecCtx, FVectorVMSerializeState *SerializeState, FVectorVMSerializeState *CmpSerializeState, VVMFn_null fn, int NumLoops) {
-	VVMSer_instruction(0, 2);
-	return fn(CT_MultipleLoops, InsPtr, BatchState, ExecCtx, SerializeState, CmpSerializeState, NumLoops);
+VM_FORCEINLINE const uint8 *VVM_Dispatch_execFn2null_1null(const bool CT_MultipleLoops, const uint8 *InsPtr, FVectorVMBatchState *BatchState, FVectorVMExecContext *ExecCtx, VVMFn_null fn, int NumLoops)
+{
+	return fn(CT_MultipleLoops, InsPtr, BatchState, ExecCtx, NumLoops);
 }
 
-VM_FORCEINLINE const uint8 *VVM_Dispatch_execFn1f_1f(const bool CT_MultipleLoops, const uint8 *InsPtr, FVectorVMBatchState *BatchState, FVectorVMExecContext *ExecCtx, FVectorVMSerializeState *SerializeState, FVectorVMSerializeState *CmpSerializeState, VVMFn_1f fn, int NumLoops) {
-	VVMSer_instruction(1, 1);
+VM_FORCEINLINE const uint8 *VVM_Dispatch_execFn1f_1f(const bool CT_MultipleLoops, const uint8 *InsPtr, FVectorVMBatchState *BatchState, FVectorVMExecContext *ExecCtx, VVMFn_1f fn, int NumLoops) {
 	uint16 *RegIndices = (uint16 *)InsPtr;
 	uint8 *P0 = BatchState->RegPtrTable[RegIndices[0]];
 	uint8 *P1 = BatchState->RegPtrTable[RegIndices[1]];
@@ -2029,8 +1787,8 @@ VM_FORCEINLINE const uint8 *VVM_Dispatch_execFn1f_1f(const bool CT_MultipleLoops
 	return InsPtr + 5;
 }
 
-VM_FORCEINLINE const uint8 *VVM_Dispatch_execFn2f_1f(const bool CT_MultipleLoops, const uint8 *InsPtr, FVectorVMBatchState *BatchState, FVectorVMExecContext *ExecCtx, FVectorVMSerializeState *SerializeState, FVectorVMSerializeState *CmpSerializeState, VVMFn_2f fn, int NumLoops) {
-	VVMSer_instruction(1, 2);
+VM_FORCEINLINE const uint8 *VVM_Dispatch_execFn2f_1f(const bool CT_MultipleLoops, const uint8 *InsPtr, FVectorVMBatchState *BatchState, FVectorVMExecContext *ExecCtx, VVMFn_2f fn, int NumLoops)
+{
 	uint16 *RegIndices = (uint16 *)InsPtr;
 	uint8 *P0 = BatchState->RegPtrTable[RegIndices[0]];
 	uint8 *P1 = BatchState->RegPtrTable[RegIndices[1]];
@@ -2061,8 +1819,8 @@ VM_FORCEINLINE const uint8 *VVM_Dispatch_execFn2f_1f(const bool CT_MultipleLoops
 	return InsPtr + 7;
 }
 
-VM_FORCEINLINE const uint8 *VVM_Dispatch_execFn3f_1f(const bool CT_MultipleLoops, const uint8 *InsPtr, FVectorVMBatchState *BatchState, FVectorVMExecContext *ExecCtx, FVectorVMSerializeState *SerializeState, FVectorVMSerializeState *CmpSerializeState, VVMFn_3f fn, int NumLoops) {
-	VVMSer_instruction(1, 3);
+VM_FORCEINLINE const uint8 *VVM_Dispatch_execFn3f_1f(const bool CT_MultipleLoops, const uint8 *InsPtr, FVectorVMBatchState *BatchState, FVectorVMExecContext *ExecCtx, VVMFn_3f fn, int NumLoops)
+{
 	uint16 *RegIndices = (uint16 *)InsPtr;
 	uint8 *P0 = BatchState->RegPtrTable[RegIndices[0]];
 	uint8 *P1 = BatchState->RegPtrTable[RegIndices[1]];
@@ -2098,8 +1856,8 @@ VM_FORCEINLINE const uint8 *VVM_Dispatch_execFn3f_1f(const bool CT_MultipleLoops
 	return InsPtr + 9;
 }
 
-VM_FORCEINLINE const uint8 *VVM_Dispatch_execFn4f_1f(const bool CT_MultipleLoops, const uint8 *InsPtr, FVectorVMBatchState *BatchState, FVectorVMExecContext *ExecCtx, FVectorVMSerializeState *SerializeState, FVectorVMSerializeState *CmpSerializeState, VVMFn_4f fn, int NumLoops) {
-	VVMSer_instruction(1, 4);
+VM_FORCEINLINE const uint8 *VVM_Dispatch_execFn4f_1f(const bool CT_MultipleLoops, const uint8 *InsPtr, FVectorVMBatchState *BatchState, FVectorVMExecContext *ExecCtx, VVMFn_4f fn, int NumLoops)
+{
 	uint16 *RegIndices = (uint16 *)InsPtr;
 	uint8 *P0 = BatchState->RegPtrTable[RegIndices[0]];
 	uint8 *P1 = BatchState->RegPtrTable[RegIndices[1]];
@@ -2140,8 +1898,8 @@ VM_FORCEINLINE const uint8 *VVM_Dispatch_execFn4f_1f(const bool CT_MultipleLoops
 	return InsPtr + 11;
 }
 
-VM_FORCEINLINE const uint8 *VVM_Dispatch_execFn5f_1f(const bool CT_MultipleLoops, const uint8 *InsPtr, FVectorVMBatchState *BatchState, FVectorVMExecContext *ExecCtx, FVectorVMSerializeState *SerializeState, FVectorVMSerializeState *CmpSerializeState, VVMFn_5f fn, int NumLoops) {
-	VVMSer_instruction(1, 5);
+VM_FORCEINLINE const uint8 *VVM_Dispatch_execFn5f_1f(const bool CT_MultipleLoops, const uint8 *InsPtr, FVectorVMBatchState *BatchState, FVectorVMExecContext *ExecCtx, VVMFn_5f fn, int NumLoops)
+{
 	uint16 *RegIndices = (uint16 *)InsPtr;
 	uint8 *P0 = BatchState->RegPtrTable[RegIndices[0]];
 	uint8 *P1 = BatchState->RegPtrTable[RegIndices[1]];
@@ -2187,8 +1945,8 @@ VM_FORCEINLINE const uint8 *VVM_Dispatch_execFn5f_1f(const bool CT_MultipleLoops
 	return InsPtr + 13;
 }
 
-VM_FORCEINLINE const uint8 *VVM_Dispatch_execFn2f_2f(const bool CT_MultipleLoops, const uint8 *InsPtr, FVectorVMBatchState *BatchState, FVectorVMExecContext *ExecCtx, FVectorVMSerializeState *SerializeState, FVectorVMSerializeState *CmpSerializeState, VVMFn_2f fn, int NumLoops) {
-	VVMSer_instruction(1, 3);
+VM_FORCEINLINE const uint8 *VVM_Dispatch_execFn2f_2f(const bool CT_MultipleLoops, const uint8 *InsPtr, FVectorVMBatchState *BatchState, FVectorVMExecContext *ExecCtx, VVMFn_2f fn, int NumLoops)
+{
 	uint16 *RegIndices = (uint16 *)InsPtr;
 	uint8 *P0 = BatchState->RegPtrTable[RegIndices[0]];
 	uint8 *P1 = BatchState->RegPtrTable[RegIndices[1]];
@@ -2223,8 +1981,8 @@ VM_FORCEINLINE const uint8 *VVM_Dispatch_execFn2f_2f(const bool CT_MultipleLoops
 	return InsPtr + 9;
 }
 
-VM_FORCEINLINE const uint8 *VVM_Dispatch_execFn1i_1i(const bool CT_MultipleLoops, const uint8 *InsPtr, FVectorVMBatchState *BatchState, FVectorVMExecContext *ExecCtx, FVectorVMSerializeState *SerializeState, FVectorVMSerializeState *CmpSerializeState, VVMFn_1i fn, int NumLoops) {
-	VVMSer_instruction(1, 1);
+VM_FORCEINLINE const uint8 *VVM_Dispatch_execFn1i_1i(const bool CT_MultipleLoops, const uint8 *InsPtr, FVectorVMBatchState *BatchState, FVectorVMExecContext *ExecCtx, VVMFn_1i fn, int NumLoops)
+{
 	uint16 *RegIndices = (uint16 *)InsPtr;
 	uint8 *P0 = BatchState->RegPtrTable[RegIndices[0]];
 	uint8 *P1 = BatchState->RegPtrTable[RegIndices[1]];
@@ -2250,8 +2008,8 @@ VM_FORCEINLINE const uint8 *VVM_Dispatch_execFn1i_1i(const bool CT_MultipleLoops
 	return InsPtr + 5;
 }
 
-VM_FORCEINLINE const uint8 *VVM_Dispatch_execFn2i_1i(const bool CT_MultipleLoops, const uint8 *InsPtr, FVectorVMBatchState *BatchState, FVectorVMExecContext *ExecCtx, FVectorVMSerializeState *SerializeState, FVectorVMSerializeState *CmpSerializeState, VVMFn_2i fn, int NumLoops) {
-	VVMSer_instruction(1, 2);
+VM_FORCEINLINE const uint8 *VVM_Dispatch_execFn2i_1i(const bool CT_MultipleLoops, const uint8 *InsPtr, FVectorVMBatchState *BatchState, FVectorVMExecContext *ExecCtx, VVMFn_2i fn, int NumLoops)
+{
 	uint16 *RegIndices = (uint16 *)InsPtr;
 	uint8 *P0 = BatchState->RegPtrTable[RegIndices[0]];
 	uint8 *P1 = BatchState->RegPtrTable[RegIndices[1]];
@@ -2282,8 +2040,8 @@ VM_FORCEINLINE const uint8 *VVM_Dispatch_execFn2i_1i(const bool CT_MultipleLoops
 	return InsPtr + 7;
 }
 
-VM_FORCEINLINE const uint8 *VVM_Dispatch_execFn3i_1i(const bool CT_MultipleLoops, const uint8 *InsPtr, FVectorVMBatchState *BatchState, FVectorVMExecContext *ExecCtx, FVectorVMSerializeState *SerializeState, FVectorVMSerializeState *CmpSerializeState, VVMFn_3i fn, int NumLoops) {
-	VVMSer_instruction(1, 3);
+VM_FORCEINLINE const uint8 *VVM_Dispatch_execFn3i_1i(const bool CT_MultipleLoops, const uint8 *InsPtr, FVectorVMBatchState *BatchState, FVectorVMExecContext *ExecCtx, VVMFn_3i fn, int NumLoops)
+{
 	uint16 *RegIndices = (uint16 *)InsPtr;
 	uint8 *P0 = BatchState->RegPtrTable[RegIndices[0]];
 	uint8 *P1 = BatchState->RegPtrTable[RegIndices[1]];
@@ -2319,8 +2077,8 @@ VM_FORCEINLINE const uint8 *VVM_Dispatch_execFn3i_1i(const bool CT_MultipleLoops
 	return InsPtr + 9;
 }
 
-VM_FORCEINLINE const uint8 *VVM_Dispatch_execFn1i_2i(const bool CT_MultipleLoops, const uint8 *InsPtr, FVectorVMBatchState *BatchState, FVectorVMExecContext *ExecCtx, FVectorVMSerializeState *SerializeState, FVectorVMSerializeState *CmpSerializeState, VVMFn_1i fn, int NumLoops) {
-	VVMSer_instruction(1, 2);
+VM_FORCEINLINE const uint8 *VVM_Dispatch_execFn1i_2i(const bool CT_MultipleLoops, const uint8 *InsPtr, FVectorVMBatchState *BatchState, FVectorVMExecContext *ExecCtx, VVMFn_1i fn, int NumLoops)
+{
 	uint16 *RegIndices = (uint16 *)InsPtr;
 	uint8 *P0 = BatchState->RegPtrTable[RegIndices[0]];
 	uint8 *P1 = BatchState->RegPtrTable[RegIndices[1]];
@@ -2384,7 +2142,7 @@ VM_FORCEINLINE VectorRegister4f VVM_Exec2f_pow                             (FVec
 VM_FORCEINLINE VectorRegister4f VVM_Exec1f_round                           (FVectorVMBatchState *BatchState, VectorRegister4f a)                                                                                { return VectorRound(a); }
 VM_FORCEINLINE VectorRegister4f VVM_Exec1f_sign                            (FVectorVMBatchState *BatchState, VectorRegister4f a)                                                                                { return VectorSign(a); }
 VM_FORCEINLINE VectorRegister4f VVM_Exec2f_step                            (FVectorVMBatchState *BatchState, VectorRegister4f a, VectorRegister4f b)                                                            { return VectorStep(VectorSubtract(b, a)); }
-VM_FORCEINLINE const uint8 *    VVM_Exec1null_random                       (VVM_NULL_FN_ARGS)                                                                                                                   { return VVM_random(CT_MultipleLoops, InsPtr, BatchState, ExecCtx, SerializeState, CmpSerializeState, NumLoops); }
+VM_FORCEINLINE const uint8 *    VVM_Exec1null_random                       (VVM_NULL_FN_ARGS)                                                                                                                   { return VVM_random(CT_MultipleLoops, InsPtr, BatchState, ExecCtx, NumLoops); }
 VM_FORCEINLINE const uint8 *    VVM_Exec0null_noise                        (VVM_NULL_FN_ARGS)                                                                                                                   { return InsPtr; }
 VM_FORCEINLINE VectorRegister4f VVM_Exec2f_cmplt                           (FVectorVMBatchState *BatchState, VectorRegister4f a, VectorRegister4f b)                                                            { return VectorCompareLT(a, b); }
 VM_FORCEINLINE VectorRegister4f VVM_Exec2f_cmple                           (FVectorVMBatchState *BatchState, VectorRegister4f a, VectorRegister4f b)                                                            { return VectorCompareLE(a, b); }
@@ -2403,7 +2161,7 @@ VM_FORCEINLINE VectorRegister4i VVM_Exec2i_maxi                            (FVec
 VM_FORCEINLINE VectorRegister4i VVM_Exec1i_absi                            (FVectorVMBatchState *BatchState, VectorRegister4i a)                                                                                { return VectorIntAbs(a); }
 VM_FORCEINLINE VectorRegister4i VVM_Exec1i_negi                            (FVectorVMBatchState *BatchState, VectorRegister4i a)                                                                                { return VectorIntNegate(a); }
 VM_FORCEINLINE VectorRegister4i VVM_Exec1i_signi                           (FVectorVMBatchState *BatchState, VectorRegister4i a)                                                                                { return VectorIntSign(a); }
-VM_FORCEINLINE const uint8 *    VVM_Exec1null_randomi                      (VVM_NULL_FN_ARGS)                                                                                                                   { return VVM_randomi(CT_MultipleLoops, InsPtr, BatchState, ExecCtx, SerializeState, CmpSerializeState, NumLoops); }
+VM_FORCEINLINE const uint8 *    VVM_Exec1null_randomi                      (VVM_NULL_FN_ARGS)                                                                                                                   { return VVM_randomi(CT_MultipleLoops, InsPtr, BatchState, ExecCtx, NumLoops); }
 VM_FORCEINLINE VectorRegister4i VVM_Exec2i_cmplti                          (FVectorVMBatchState *BatchState, VectorRegister4i a, VectorRegister4i b)                                                            { return VectorIntCompareLT(a, b); }
 VM_FORCEINLINE VectorRegister4i VVM_Exec2i_cmplei                          (FVectorVMBatchState *BatchState, VectorRegister4i a, VectorRegister4i b)                                                            { return VectorIntCompareLE(a, b); }
 VM_FORCEINLINE VectorRegister4i VVM_Exec2i_cmpgti                          (FVectorVMBatchState *BatchState, VectorRegister4i a, VectorRegister4i b)                                                            { return VectorIntCompareGT(a, b); }
@@ -2430,7 +2188,7 @@ VM_FORCEINLINE const uint8 *    VVM_Exec0null_outputdata_float             (VVM_
 VM_FORCEINLINE const uint8 *    VVM_Exec0null_outputdata_int32             (VVM_NULL_FN_ARGS)                                                                                                                   { return VVM_Output32(CT_MultipleLoops, InsPtr, BatchState, ExecCtx); }
 VM_FORCEINLINE const uint8 *    VVM_Exec0null_outputdata_half              (VVM_NULL_FN_ARGS)                                                                                                                   { return VVM_Output16(CT_MultipleLoops, InsPtr, BatchState, ExecCtx); }
 VM_FORCEINLINE const uint8 *    VVM_Exec1null_acquireindex                 (VVM_NULL_FN_ARGS)                                                                                                                   { return VVM_acquireindex(CT_MultipleLoops, InsPtr, BatchState, ExecCtx); }
-VM_FORCEINLINE const uint8 *    VVM_Exec0null_external_func_call           (VVM_NULL_FN_ARGS)                                                                                                                   { return VVM_external_func_call(CT_MultipleLoops, InsPtr, BatchState, ExecCtx, SerializeState, CmpSerializeState, NumLoops); }
+VM_FORCEINLINE const uint8 *    VVM_Exec0null_external_func_call           (VVM_NULL_FN_ARGS)                                                                                                                   { return VVM_external_func_call(CT_MultipleLoops, InsPtr, BatchState, ExecCtx, NumLoops); }
 VM_FORCEINLINE const uint8 *    VVM_Exec1null_exec_index                   (VVM_NULL_FN_ARGS)                                                                                                                   { return VVM_exec_index(CT_MultipleLoops, InsPtr, BatchState, NumLoops); }
 VM_FORCEINLINE const uint8 *    VVM_Exec0null_noise2D                      (VVM_NULL_FN_ARGS)                                                                                                                   { return InsPtr; }
 VM_FORCEINLINE const uint8 *    VVM_Exec0null_noise3D                      (VVM_NULL_FN_ARGS)                                                                                                                   { return InsPtr; }
@@ -2514,8 +2272,8 @@ VM_FORCEINLINE VectorRegister4i VVM_Exec3i_bit_rshift_bit_and              (FVec
 VM_FORCEINLINE VectorRegister4f VVM_Exec2f_neg_cmplt                       (FVectorVMBatchState *BatchState, VectorRegister4f a, VectorRegister4f b)                                                            { return VectorCompareLT(VectorNegate(a), b); }
 VM_FORCEINLINE VectorRegister4i VVM_Exec3i_bit_or_muli                     (FVectorVMBatchState *BatchState, VectorRegister4i a, VectorRegister4i b, VectorRegister4i c)                                        { return VectorIntMultiply(VectorIntOr(a, b), c); }
 VM_FORCEINLINE VectorRegister4i VVM_Exec3i_bit_lshift_bit_or               (FVectorVMBatchState *BatchState, VectorRegister4i a, VectorRegister4i b, VectorRegister4i c)                                        { return VectorIntOr(VVMIntLShift(a, b), c); }
-VM_FORCEINLINE const uint8 *    VVM_Exec2null_random_add                   (VVM_NULL_FN_ARGS)                                                                                                                   { return VVM_random_add(CT_MultipleLoops, InsPtr, BatchState, ExecCtx, SerializeState, CmpSerializeState, NumLoops); }
-VM_FORCEINLINE const uint8 *    VVM_Exec1null_random_2x                    (VVM_NULL_FN_ARGS)                                                                                                                   { return VVM_random_2x(CT_MultipleLoops, InsPtr, BatchState, ExecCtx, SerializeState, CmpSerializeState, NumLoops); }
+VM_FORCEINLINE const uint8 *    VVM_Exec2null_random_add                   (VVM_NULL_FN_ARGS)                                                                                                                   { return VVM_random_add(CT_MultipleLoops, InsPtr, BatchState, ExecCtx, NumLoops); }
+VM_FORCEINLINE const uint8 *    VVM_Exec1null_random_2x                    (VVM_NULL_FN_ARGS)                                                                                                                   { return VVM_random_2x(CT_MultipleLoops, InsPtr, BatchState, ExecCtx, NumLoops); }
 VM_FORCEINLINE VectorRegister4i VVM_Exec2i_max_f2i                         (FVectorVMBatchState *BatchState, VectorRegister4i a, VectorRegister4i b)                                                            { return VectorFloatToInt(VectorMax(VectorCastIntToFloat(a), VectorCastIntToFloat(b))); }
 VM_FORCEINLINE VectorRegister4f VVM_Exec4f_select_mul                      (FVectorVMBatchState *BatchState, VectorRegister4f a, VectorRegister4f b, VectorRegister4f c, VectorRegister4f d)                    { return VectorMultiply(VectorSelect(a, b, c), d); }
 VM_FORCEINLINE VectorRegister4f VVM_Exec4f_select_add                      (FVectorVMBatchState *BatchState, VectorRegister4f a, VectorRegister4f b, VectorRegister4f c, VectorRegister4f d)                    { return VectorAdd(VectorSelect(a, b, c), d); }
@@ -2591,7 +2349,7 @@ static void SetBatchPointersForCorrectChunkOffsets(FVectorVMExecContext *ExecCtx
 	}
 }
 
-static void ExecChunkSingleLoop(FVectorVMExecContext *ExecCtx, FVectorVMBatchState *BatchState, FVectorVMSerializeState *SerializeState, FVectorVMSerializeState *CmpSerializeState)
+static void ExecChunkSingleLoop(FVectorVMExecContext *ExecCtx, FVectorVMBatchState *BatchState)
 {
 	static const int      NumLoops       = 1;
 	static constexpr bool CT_MultipleLoops = false;
@@ -2601,25 +2359,22 @@ static void ExecChunkSingleLoop(FVectorVMExecContext *ExecCtx, FVectorVMBatchSta
 	const uint8 *InsPtr        = ExecCtx->VVMState->Bytecode + 1;
 	const uint8 *InsPtrEnd     = InsPtr + ExecCtx->VVMState->NumBytecodeBytes;
 
-	VVMSer_chunkStartExp(SerializeState, ChunkIdx, 0)
 #	if VECTORVM_DEBUG_PRINTF
 	int DbgPfInsCount = 0;
 #	endif //VECTORVM_DEBUG_PRINTF
 	for (;;)
 	{
-		VVMSer_insStartExp(SerializeState)
-		VVM_OP_START {
-#			define VVM_OP_XM(OpCode, Cat, NumInputs, NumOutputs, Type, ...) VVM_OP_CASE(OpCode) InsPtr = VVM_Dispatch_execFn##NumInputs##Type##_##NumOutputs##Type(CT_MultipleLoops, InsPtr, BatchState, ExecCtx, SerializeState, CmpSerializeState, VVM_Exec##NumInputs##Type##_##OpCode, NumLoops); VVM_OP_NEXT;
+		VVM_OP_START
+		{
+#			define VVM_OP_XM(OpCode, Cat, NumInputs, NumOutputs, Type, ...) VVM_OP_CASE(OpCode) InsPtr = VVM_Dispatch_execFn##NumInputs##Type##_##NumOutputs##Type(CT_MultipleLoops, InsPtr, BatchState, ExecCtx, VVM_Exec##NumInputs##Type##_##OpCode, NumLoops); VVM_OP_NEXT;
 			VVM_OP_XM_LIST
 #			undef VVM_OP_XM
 		}
-		VVMSer_insEndExp(SerializeState, (int)(SerializeState->Running.StartOpPtr - ExecCtx->VVMState->Bytecode), (int)(InsPtr - SerializeState->Running.StartOpPtr));
 	}
-	done_loop:
-	VVMSer_chunkEndExp(SerializeState);
+	done_loop: ; // breaking out of the above loop
 }
 
-static void ExecChunkMultipleLoops(FVectorVMExecContext *ExecCtx, FVectorVMBatchState *BatchState, int NumLoops, FVectorVMSerializeState *SerializeState, FVectorVMSerializeState *CmpSerializeState)
+static void ExecChunkMultipleLoops(FVectorVMExecContext *ExecCtx, FVectorVMBatchState *BatchState, int NumLoops)
 {
 	static constexpr bool CT_MultipleLoops = true;
 	SetBatchPointersForCorrectChunkOffsets(ExecCtx, BatchState);
@@ -2627,22 +2382,19 @@ static void ExecChunkMultipleLoops(FVectorVMExecContext *ExecCtx, FVectorVMBatch
 	const uint8 *InsPtr        = ExecCtx->VVMState->Bytecode + 1;
 	const uint8 *InsPtrEnd     = InsPtr + ExecCtx->VVMState->NumBytecodeBytes;
 
-	VVMSer_chunkStartExp(SerializeState, ChunkIdx, 0)
 #	if VECTORVM_DEBUG_PRINTF
 	int DbgPfInsCount = 0;
 #	endif //VECTORVM_DEBUG_PRINTF
 	for (;;)
 	{
-		VVMSer_insStartExp(SerializeState)
-		VVM_OP_START {
-#			define VVM_OP_XM(OpCode, Cat, NumInputs, NumOutputs, Type, ...) VVM_OP_CASE(OpCode) InsPtr = VVM_Dispatch_execFn##NumInputs##Type##_##NumOutputs##Type(CT_MultipleLoops, InsPtr, BatchState, ExecCtx, SerializeState, CmpSerializeState, VVM_Exec##NumInputs##Type##_##OpCode, NumLoops); VVM_OP_NEXT;
+		VVM_OP_START
+		{
+#			define VVM_OP_XM(OpCode, Cat, NumInputs, NumOutputs, Type, ...) VVM_OP_CASE(OpCode) InsPtr = VVM_Dispatch_execFn##NumInputs##Type##_##NumOutputs##Type(CT_MultipleLoops, InsPtr, BatchState, ExecCtx, VVM_Exec##NumInputs##Type##_##OpCode, NumLoops); VVM_OP_NEXT;
 			VVM_OP_XM_LIST
 #			undef VVM_OP_XM
 		}
-		VVMSer_insEndExp(SerializeState, (int)(SerializeState->Running.StartOpPtr - ExecCtx->VVMState->Bytecode), (int)(InsPtr - SerializeState->Running.StartOpPtr));
 	}
-	done_loop:
-	VVMSer_chunkEndExp(SerializeState);
+	done_loop: ; // breaking out of the above loop
 }
 
 #undef VVM_Dispatch_execFn0done_0done
@@ -2650,7 +2402,7 @@ static void ExecChunkMultipleLoops(FVectorVMExecContext *ExecCtx, FVectorVMBatch
 #undef VVM_OP_CASE
 #undef VVM_OP_NEXT
 
-void ExecVectorVMState(FVectorVMExecContext *ExecCtx, FVectorVMSerializeState *SerializeState, FVectorVMSerializeState *CmpSerializeState)
+void ExecVectorVMState(FVectorVMExecContext *ExecCtx)
 {
 	if (ExecCtx->ExtFunctionTable.Num() != ExecCtx->VVMState->NumExtFunctions) {
 		return;
@@ -2726,18 +2478,6 @@ void ExecVectorVMState(FVectorVMExecContext *ExecCtx, FVectorVMSerializeState *S
 		ExecCtx->VVMState->NumInstancesExecCached = ExecCtx->NumInstances;
 	}
 
-	VVMSer_initSerializationState(ExecCtx, SerializeState, SerializeState->OptimizeCtx, SerializeState->Flags | VVMSer_OptimizedBytecode);
-
-#if VECTORVM_SUPPORTS_SERIALIZATION && !defined(VVM_SERIALIZE_NO_WRITE)
-	uint64 StartTime = FPlatformTime::Cycles64();
-	if (SerializeState)
-	{
-		SerializeState->ExecDt = 0;
-		SerializeState->SerializeDt = 0;
-	}
-#endif
-	SerializeVectorVMInputDataSets(SerializeState, ExecCtx);
-
 	{
 		if (ExecCtx->VVMState->Bytecode == nullptr) {
 			return;
@@ -2763,16 +2503,14 @@ void ExecVectorVMState(FVectorVMExecContext *ExecCtx, FVectorVMSerializeState *S
 			
 			if (NumLoops == 1)
 			{
-				ExecChunkSingleLoop(ExecCtx, BatchState, SerializeState, CmpSerializeState);
+				ExecChunkSingleLoop(ExecCtx, BatchState);
 			}
 			else if (NumLoops >= 1)
 			{
-				ExecChunkMultipleLoops(ExecCtx, BatchState, NumLoops, SerializeState, CmpSerializeState);
+				ExecChunkMultipleLoops(ExecCtx, BatchState, NumLoops);
 			}
 		}
 	
-		VVMSer_batchEndExp(SerializeState);
-
 		if (BatchState->ChunkLocalData.RandCounters)
 		{
 			FMemory::Free(BatchState->ChunkLocalData.RandCounters);
@@ -2785,17 +2523,6 @@ void ExecVectorVMState(FVectorVMExecContext *ExecCtx, FVectorVMSerializeState *S
 	{
 		ExecCtx->DataSets[i].DataSetAccessIndex = ExecCtx->VVMState->NumOutputPerDataSet[i] - 1;
 	}
-
-#if VECTORVM_SUPPORTS_SERIALIZATION && !defined(VVM_SERIALIZE_NO_WRITE)
-	uint64 EndTime = FPlatformTime::Cycles64();
-	if (SerializeState)
-	{
-		SerializeState->ExecDt = EndTime - StartTime;
-	}
-#endif //VECTORVM_SUPPORTS_SERIALIZATION
-
-
-	SerializeVectorVMOutputDataSets(SerializeState, ExecCtx);
 }
 
 #undef VVM_MIN
diff --git a/Engine/Source/Runtime/VectorVM/Private/VectorVMExperimentalSerialization.cpp b/Engine/Source/Runtime/VectorVM/Private/VectorVMExperimentalSerialization.cpp
deleted file mode 100644
diff --git a/Engine/Source/Runtime/VectorVM/Public/VectorVM.h b/Engine/Source/Runtime/VectorVM/Public/VectorVM.h
--- a/Engine/Source/Runtime/VectorVM/Public/VectorVM.h
+++ b/Engine/Source/Runtime/VectorVM/Public/VectorVM.h
@@ -15,8 +15,6 @@
 #include "VectorVMExperimental.h"
 #include "VectorVMCommon.h"
 
-struct FVectorVMSerializeState;
-
 using FVectorVMExternalFunctionContext = FVectorVMExternalFunctionContextExperimental;
 
 namespace VectorVM
diff --git a/Engine/Source/Runtime/VectorVM/Public/VectorVMExperimental.h b/Engine/Source/Runtime/VectorVM/Public/VectorVMExperimental.h
--- a/Engine/Source/Runtime/VectorVM/Public/VectorVMExperimental.h
+++ b/Engine/Source/Runtime/VectorVM/Public/VectorVMExperimental.h
@@ -1,7 +1,6 @@
 // Copyright Epic Games, Inc. All Rights Reserved.
 #pragma once
 
-#include "VectorVMSerialization.h"
 #include "VectorVMCommon.h"
 #include "Math/RandomStream.h"
 
@@ -127,7 +126,7 @@ enum EVectorVMFlags
 	VVMFlag_DataMapCacheSetup        = 1 << 4,
 };
 
-//prototypes for serialization that are required whether or not serialization is enabled
+//prototypes for optimizer
 typedef uint32 (VectorVMOptimizeErrorCallback)     (struct FVectorVMOptimizeContext *OptimizeContext, uint32 ErrorFlags);	//return new error flags
 
 union FVecReg {
@@ -436,7 +435,7 @@ public:
 //normal functions
 VECTORVM_API FVectorVMState * AllocVectorVMState                     (FVectorVMOptimizeContext *OptimizeCtx);
 VECTORVM_API void             FreeVectorVMState                      (FVectorVMState *VectorVMState);
-VECTORVM_API void             ExecVectorVMState                      (FVectorVMExecContext *ExecCtx, FVectorVMSerializeState *SerializeState, FVectorVMSerializeState *CmpSerializeState);
+VECTORVM_API void             ExecVectorVMState                      (FVectorVMExecContext *ExecCtx);
 
 #if WITH_EDITORONLY_DATA
 //optimize functions
@@ -447,12 +446,6 @@ VECTORVM_API void             GenerateHumanReadableVectorVMScript    (const FVec
 #endif
 VECTORVM_API void             ReinterpretVectorVMOptimizeContextData (TConstArrayView<uint8> ContextData, FVectorVMOptimizeContext& Context);
 
-//serialize functions
-VECTORVM_API uint32           SerializeVectorVMInputDataSets         (FVectorVMSerializeState *SerializeState, FVectorVMExecContext *ExecContext);
-VECTORVM_API uint32           SerializeVectorVMOutputDataSets        (FVectorVMSerializeState *SerializeState, FVectorVMExecContext *ExecContext);
-VECTORVM_API void             SerializeVectorVMWriteToFile           (FVectorVMSerializeState *SerializeState, uint8 WhichStateWritten, const wchar_t *Filename);
-VECTORVM_API void             FreeVectorVMSerializeState             (FVectorVMSerializeState *SerializeState);
-
 //debug helper functions
 VECTORVM_API uint16           VVMDbgGetRemappedInputIdx              (FVectorVMState *VectorVMState, int InputRegisterIdx, int DataSetIdx, bool IntReg, uint16 *OptOutIdxInRegTable);
 VECTORVM_API uint16           VVMDbgGetRemappedOutputIdx             (FVectorVMState *VectorVMState, int OutputRegisterIdx, int DataSetIdx, bool IntReg, uint16 *OptOutIdxInRegTable);
diff --git a/Engine/Source/Runtime/VectorVM/Public/VectorVMSerialization.h b/Engine/Source/Runtime/VectorVM/Public/VectorVMSerialization.h
deleted file mode 100644
diff --git a/Engine/Source/Runtime/VectorVM/VectorVM.Build.cs b/Engine/Source/Runtime/VectorVM/VectorVM.Build.cs
--- a/Engine/Source/Runtime/VectorVM/VectorVM.Build.cs
+++ b/Engine/Source/Runtime/VectorVM/VectorVM.Build.cs
@@ -16,7 +16,6 @@ public class VectorVM : ModuleRules
         PublicDefinitions.AddRange(
             new string[]
             {
-                "VECTORVM_SUPPORTS_SERIALIZATION=0",
                 "VECTORVM_DEBUG_PRINTF=0"
             });
     }
