Fix typo with VectorVM cleanup
#tests Compiled file

diff --git a/Engine/Source/Runtime/VectorVM/Private/VectorVM.cpp b/Engine/Source/Runtime/VectorVM/Private/VectorVM.cpp
--- a/Engine/Source/Runtime/VectorVM/Private/VectorVM.cpp
+++ b/Engine/Source/Runtime/VectorVM/Private/VectorVM.cpp
@@ -260,7 +260,7 @@ UEnum* g_VectorVMEnumOperandObj = nullptr;
 static const char *VVM_OP_NAMES[] {
 	VVM_OP_XM_LIST
 };
-#undef VVM_OP_XMd 
+#undef VVM_OP_XM
 
 FString VectorVM::GetOpName(EVectorVMOp Op)
 {
