Fix for op names in the VectorVM when generating the code view
#tests Looked at dumped bytecode from commandlet

diff --git a/Engine/Source/Runtime/VectorVM/Private/VectorVM.cpp b/Engine/Source/Runtime/VectorVM/Private/VectorVM.cpp
--- a/Engine/Source/Runtime/VectorVM/Private/VectorVM.cpp
+++ b/Engine/Source/Runtime/VectorVM/Private/VectorVM.cpp
@@ -26,7 +26,7 @@ uint8 GetNumOpCodes()
 #if WITH_EDITOR
 UEnum* VectorVMOperandEnum = nullptr;
 
-#define VVM_OP_XM(n, ...) TEXT("#n"),
+#define VVM_OP_XM(n, ...) #n,
 static const TArray<FString> VVM_OP_NAMES
 {
 	VVM_OP_XM_LIST
