Fixes overflowing u8 limit of external functions and introduces a shader version guid for the VM shader format to be bumped whenever the backend compiler or the interchange format changes.

diff --git a/Engine/Source/Developer/ShaderFormatVectorVM/Private/ShaderFormatVectorVM.cpp b/Engine/Source/Developer/ShaderFormatVectorVM/Private/ShaderFormatVectorVM.cpp
--- a/Engine/Source/Developer/ShaderFormatVectorVM/Private/ShaderFormatVectorVM.cpp
+++ b/Engine/Source/Developer/ShaderFormatVectorVM/Private/ShaderFormatVectorVM.cpp
@@ -21,6 +21,7 @@ extern bool CompileVectorVMShader(
 	const FString& WorkingDirectory);
 
 static FName NAME_VVM_1_0(TEXT("VVM_1_0"));
+static const FGuid UE_SHADER_VECTORVM_VER = FGuid("51C3EC95-421E-43B1-8AA9-61B6D4EDC1BF");
 
 class FShaderFormatVectorVM : public UE::ShaderCompilerCommon::FBaseShaderFormat
 {
@@ -47,8 +48,14 @@ public:
 		{
 			check(0);
 		}
+
 		const uint16 Version = ((HLSLCC_VersionMinor & 0xff) << 8) | (VVMVersion & 0xff);
-		return Version;
+
+		uint32 Result = GetTypeHash(HLSLCC_VersionMinor);
+		Result = HashCombine(Result, GetTypeHash(VVMVersion));
+		Result = HashCombine(Result, GetTypeHash(UE_SHADER_VECTORVM_VER));
+
+		return Result;
 	}
 	virtual void GetSupportedFormats(TArray<FName>& OutFormats) const override
 	{
diff --git a/Engine/Source/Developer/ShaderFormatVectorVM/Private/ir_vm_gen_bytecode_visitor.cpp b/Engine/Source/Developer/ShaderFormatVectorVM/Private/ir_vm_gen_bytecode_visitor.cpp
--- a/Engine/Source/Developer/ShaderFormatVectorVM/Private/ir_vm_gen_bytecode_visitor.cpp
+++ b/Engine/Source/Developer/ShaderFormatVectorVM/Private/ir_vm_gen_bytecode_visitor.cpp
@@ -1419,7 +1419,8 @@ struct op_external_func : public op_base
 
 		emit_byte((uint8)EVectorVMOp::external_func_call, bytecode);
 
-		emit_byte(function_index, bytecode);
+		check(function_index < std::numeric_limits<uint16>::max());
+		emit_short(function_index, bytecode);
 
 		for (variable_info_node* input : inputs)
 		{
@@ -1461,6 +1462,15 @@ struct op_external_func : public op_base
 
 		//We also add this call to the compiler output here as we're now sure it's been called.
 		function_index = CompilationOutput.CalledVMFunctionTable.AddDefaulted();
+
+		// we're also reserving one slot for -1 as an invalid index
+		constexpr unsigned int MaxExternalFunctionCount = std::numeric_limits<uint16>::max() - 1;
+		if (function_index >= MaxExternalFunctionCount)
+		{
+			_mesa_glsl_error(parse_state, "Exceeded maximum number (%d) of available external function calls.", MaxExternalFunctionCount);
+			return false;
+		}
+
 		check(function_index != INDEX_NONE);
 		FCalledVMFunction& CalledFunc = CompilationOutput.CalledVMFunctionTable[function_index];
 		CalledFunc.Name = sig->function_name();
@@ -2335,7 +2345,7 @@ class ir_gen_vvm_visitor : public ir_hierarchical_visitor
 				ordered_ops[i]->validate(parse_state, i);
 			}
 		}
-		
+
 		if (!parse_state->error)
 		{
 			//Now emit the bytecode
@@ -2355,7 +2365,6 @@ class ir_gen_vvm_visitor : public ir_hierarchical_visitor
 		FString OpStr;
 		OutNumOps = ordered_ops.Num();
 
-
 		//First go through all the parameters.
 		FString OpsConstantTable;
 		uint32 NumConstants = 0;
diff --git a/Engine/Source/Runtime/Core/Private/UObject/FortniteMainBranchObjectVersion.cpp b/Engine/Source/Runtime/Core/Private/UObject/FortniteMainBranchObjectVersion.cpp
--- a/Engine/Source/Runtime/Core/Private/UObject/FortniteMainBranchObjectVersion.cpp
+++ b/Engine/Source/Runtime/Core/Private/UObject/FortniteMainBranchObjectVersion.cpp
@@ -11,7 +11,7 @@ TMap<FGuid, FGuid> FFortniteMainBranchObjectVersion::GetSystemGuids()
 	SystemGuids.Add(DevGuids.MATERIALSHADERMAP_DERIVEDDATA_VER, FGuid("F051877974B44D409435F321BDB39DBB"));
 	SystemGuids.Add(DevGuids.NANITE_DERIVEDDATA_VER, FGuid("B169793F25FF4493A55451FDC02A4AC1"));
 	SystemGuids.Add(DevGuids.NIAGARASHADERMAP_DERIVEDDATA_VER, FGuid("6360A977062842A29ED30E3A7ACB0E64"));
-	SystemGuids.Add(DevGuids.Niagara_LatestScriptCompileVersion, FGuid("97F0FD7149D90D42B65D6557AF477D1A"));
+	SystemGuids.Add(DevGuids.Niagara_LatestScriptCompileVersion, FGuid("783FF3700BB641C799B2D43552FBF91E"));
 	SystemGuids.Add(DevGuids.SkeletalMeshDerivedDataVersion, FGuid("C998F2F3B2884A02B7E290B8C25835F0"));
 	SystemGuids.Add(DevGuids.STATICMESH_DERIVEDDATA_VER, FGuid("3ABF32149F9D4D838750368AB95573BE"));
 	SystemGuids.Add(DevGuids.MaterialTranslationDDCVersion, FGuid("C6B481BEACC04585B05635E3EC491F99"));
diff --git a/Engine/Source/Runtime/VectorVM/Private/VectorVMOptimizer.cpp b/Engine/Source/Runtime/VectorVM/Private/VectorVMOptimizer.cpp
--- a/Engine/Source/Runtime/VectorVM/Private/VectorVMOptimizer.cpp
+++ b/Engine/Source/Runtime/VectorVM/Private/VectorVMOptimizer.cpp
@@ -912,9 +912,11 @@ uint32 OptimizeVectorVMScript(const uint8 *InBytecode, int InBytecodeLen, FVecto
 			break;
 			case EVectorVMOp::external_func_call:
 			{
-				uint8 ExtFnIdx = *OpPtrIn;
+				uint16 ExtFnIdx = *reinterpret_cast<const uint16*>(OpPtrIn);
 				check(ExtFnIdx < NumExtFns);
 
+				const uint16* ExtFuncVecIndices = reinterpret_cast<const uint16*>(OpPtrIn + sizeof(ExtFnIdx));
+
 				Instruction->RegPtrOffset           = OptContext->Intermediate.NumRegistersUsed;
 				Instruction->ExtFnCall.ExtFnIdx     = ExtFnIdx;
 				Instruction->ExtFnCall.NumInputs    = ExtFnIOData[ExtFnIdx].NumInputs;
@@ -922,31 +924,31 @@ uint32 OptimizeVectorVMScript(const uint8 *InBytecode, int InBytecodeLen, FVecto
 				
 				for (int i = 0; i < ExtFnIOData[ExtFnIdx].NumInputs; ++i)
 				{
-					if (VecIndices[i] == 0xFFFF)
+					if (ExtFuncVecIndices[i] == 0xFFFF)
 					{
 						VVMPushRegUsage(0xFFFF, VVM_RT_INVALID, 0);
 					}
 					else
 					{
-						if (VecIndices[i] & 0x8000) //register: high bit means input is a register in the original bytecode
+						if (ExtFuncVecIndices[i] & 0x8000) //register: high bit means input is a register in the original bytecode
 						{ 
-							VVMPushRegUsage(VecIndices[i] & 0x7FFF, VVM_RT_TEMPREG, 0);
+							VVMPushRegUsage(ExtFuncVecIndices[i] & 0x7FFF, VVM_RT_TEMPREG, 0);
 						}
 						else //constant
 						{
-							VVMPushRegUsage(VecIndices[i], VVM_RT_CONST, 0);
+							VVMPushRegUsage(ExtFuncVecIndices[i], VVM_RT_CONST, 0);
 						}
 					}
 				}
 				for (int i = 0; i < ExtFnIOData[ExtFnIdx].NumOutputs; ++i)
 				{
 					int Idx = ExtFnIOData[ExtFnIdx].NumInputs + i;
-					check((VecIndices[Idx] & 0x8000) == 0 || VecIndices[Idx] == 0xFFFF); //can't output to a const... 0xFFFF is invalid
-					VVMPushRegUsage(VecIndices[Idx], VVM_RT_TEMPREG, 1);
+					check((ExtFuncVecIndices[Idx] & 0x8000) == 0 || ExtFuncVecIndices[Idx] == 0xFFFF); //can't output to a const... 0xFFFF is invalid
+					VVMPushRegUsage(ExtFuncVecIndices[Idx], VVM_RT_TEMPREG, 1);
 				}
 				OptContext->MaxExtFnUsed      = FMath::Max(OptContext->MaxExtFnUsed, (int32)ExtFnIdx);
 				OptContext->MaxExtFnRegisters = FMath::Max(OptContext->MaxExtFnRegisters, (uint32)(ExtFnIOData[ExtFnIdx].NumInputs + ExtFnIOData[ExtFnIdx].NumOutputs));
-				OpPtrIn += 1 + (ExtFnIOData[ExtFnIdx].NumInputs + ExtFnIOData[ExtFnIdx].NumOutputs) * 2;
+				OpPtrIn += sizeof(ExtFnIdx) + (ExtFnIOData[ExtFnIdx].NumInputs + ExtFnIOData[ExtFnIdx].NumOutputs) * 2;
 			}
 			break;
 			case EVectorVMOp::exec_index:
