- Stub VM noise function to fix crash when using in the VM
- The old VM had an implementation for 3D but not 1/2D, and the function does not work on the GPU
- Next step will be to deprecate the function and remove, it's better to use a data interface

#tests Use the noise function in a scratchpad, observe crash before

#Synced-CL 49372769

diff --git a/Engine/Source/Runtime/VectorVM/Private/VectorVMOptimizer.cpp b/Engine/Source/Runtime/VectorVM/Private/VectorVMOptimizer.cpp
--- a/Engine/Source/Runtime/VectorVM/Private/VectorVMOptimizer.cpp
+++ b/Engine/Source/Runtime/VectorVM/Private/VectorVMOptimizer.cpp
@@ -793,7 +793,7 @@ uint32 OptimizeVectorVMScript(const uint8 *InBytecode, int InBytecodeLen, FVecto
 			case EVectorVMOp::sign:                                 VVMOptimizeVecIns1;             break;
 			case EVectorVMOp::step:                                 VVMOptimizeVecIns2;             break;
 			case EVectorVMOp::random:                               VVMOptimizeVecIns1;             break;
-			case EVectorVMOp::noise:                                check(false);                   break;
+			case EVectorVMOp::noise:                                VVMOptimizeVecIns1;             break;
 			case EVectorVMOp::cmplt:                                VVMOptimizeVecIns2;             break;
 			case EVectorVMOp::cmple:                                VVMOptimizeVecIns2;             break;
 			case EVectorVMOp::cmpgt:                                VVMOptimizeVecIns2;             break;
@@ -956,8 +956,8 @@ uint32 OptimizeVectorVMScript(const uint8 *InBytecode, int InBytecodeLen, FVecto
 				VVMPushRegUsage(*OpPtrIn, VVM_RT_TEMPREG, 1);
 				OpPtrIn += 2;
 				break;
-			case EVectorVMOp::noise2D:                              check(false);               break;
-			case EVectorVMOp::noise3D:                              check(false);               break;
+			case EVectorVMOp::noise2D:                              VVMOptimizeVecIns2;             break;
+			case EVectorVMOp::noise3D:                              VVMOptimizeVecIns3;             break;
 			case EVectorVMOp::enter_stat_scope:						
 				Instruction->Stat.ID = *OpPtrIn;                    OpPtrIn += 2;               break;
 			case EVectorVMOp::exit_stat_scope:                                                  break;
diff --git a/Engine/Source/Runtime/VectorVM/Private/VectorVMRuntime.cpp b/Engine/Source/Runtime/VectorVM/Private/VectorVMRuntime.cpp
--- a/Engine/Source/Runtime/VectorVM/Private/VectorVMRuntime.cpp
+++ b/Engine/Source/Runtime/VectorVM/Private/VectorVMRuntime.cpp
@@ -1110,6 +1110,21 @@ static VM_FORCEINLINE const uint8 *VVM_random_2x(const bool CT_MultipleLoops, co
 	return InsPtr + 7;
 }
 
+static VM_FORCEINLINE VectorRegister4f VVM_noise(VectorRegister4f a)
+{
+	return GlobalVectorConstants::FloatZero;
+}
+
+static VM_FORCEINLINE VectorRegister4f VVM_noise(VectorRegister4f x, VectorRegister4f y)
+{
+	return GlobalVectorConstants::FloatZero;
+}
+
+static VM_FORCEINLINE VectorRegister4f VVM_noise(VectorRegister4f x, VectorRegister4f y, VectorRegister4f z)
+{
+	return GlobalVectorConstants::FloatZero;
+}
+
 VM_FORCEINLINE const uint8 *VVM_half_to_float(const bool CT_MultipleLoops, const uint8 *InsPtr, FVectorVMBatchState *BatchState, int NumLoops)
 {
 	uint16* RegIndices = (uint16*)InsPtr;
@@ -1938,7 +1953,7 @@ VM_FORCEINLINE VectorRegister4f VVM_Exec1f_round                           (FVec
 VM_FORCEINLINE VectorRegister4f VVM_Exec1f_sign                            (FVectorVMBatchState *BatchState, VectorRegister4f a)                                                                                { return VectorSign(a); }
 VM_FORCEINLINE VectorRegister4f VVM_Exec2f_step                            (FVectorVMBatchState *BatchState, VectorRegister4f a, VectorRegister4f b)                                                            { return VectorStep(VectorSubtract(b, a)); }
 VM_FORCEINLINE const uint8 *    VVM_Exec1null_random                       (VVM_NULL_FN_ARGS)                                                                                                                   { return VVM_random(CT_MultipleLoops, InsPtr, BatchState, ExecCtx, NumLoops); }
-VM_FORCEINLINE const uint8 *    VVM_Exec0null_noise                        (VVM_NULL_FN_ARGS)                                                                                                                   { return InsPtr; }
+VM_FORCEINLINE VectorRegister4f VVM_Exec1f_noise                           (FVectorVMBatchState* BatchState, VectorRegister4f a)																				{ return VVM_noise(a); }
 VM_FORCEINLINE VectorRegister4f VVM_Exec2f_cmplt                           (FVectorVMBatchState *BatchState, VectorRegister4f a, VectorRegister4f b)                                                            { return VectorCompareLT(a, b); }
 VM_FORCEINLINE VectorRegister4f VVM_Exec2f_cmple                           (FVectorVMBatchState *BatchState, VectorRegister4f a, VectorRegister4f b)                                                            { return VectorCompareLE(a, b); }
 VM_FORCEINLINE VectorRegister4f VVM_Exec2f_cmpgt                           (FVectorVMBatchState *BatchState, VectorRegister4f a, VectorRegister4f b)                                                            { return VectorCompareGT(a, b); }
@@ -1985,8 +2000,8 @@ VM_FORCEINLINE const uint8 *    VVM_Exec0null_outputdata_half              (VVM_
 VM_FORCEINLINE const uint8 *    VVM_Exec1null_acquireindex                 (VVM_NULL_FN_ARGS)                                                                                                                   { return VVM_acquireindex(CT_MultipleLoops, InsPtr, BatchState, ExecCtx); }
 VM_FORCEINLINE const uint8 *    VVM_Exec0null_external_func_call           (VVM_NULL_FN_ARGS)                                                                                                                   { return VVM_external_func_call(CT_MultipleLoops, InsPtr, BatchState, ExecCtx, NumLoops); }
 VM_FORCEINLINE const uint8 *    VVM_Exec1null_exec_index                   (VVM_NULL_FN_ARGS)                                                                                                                   { return VVM_exec_index(CT_MultipleLoops, InsPtr, BatchState, NumLoops); }
-VM_FORCEINLINE const uint8 *    VVM_Exec0null_noise2D                      (VVM_NULL_FN_ARGS)                                                                                                                   { return InsPtr; }
-VM_FORCEINLINE const uint8 *    VVM_Exec0null_noise3D                      (VVM_NULL_FN_ARGS)                                                                                                                   { return InsPtr; }
+VM_FORCEINLINE VectorRegister4f VVM_Exec2f_noise2D                         (FVectorVMBatchState* BatchState, VectorRegister4f x, VectorRegister4f y)															{ return VVM_noise(x, y); }
+VM_FORCEINLINE VectorRegister4f VVM_Exec3f_noise3D                         (FVectorVMBatchState* BatchState, VectorRegister4f x, VectorRegister4f y, VectorRegister4f z)										{ return VVM_noise(x, y, z); }
 VM_FORCEINLINE const uint8 *    VVM_Exec1null_enter_stat_scope             (VVM_NULL_FN_ARGS)                                                                                                                   { return InsPtr + 2; }
 VM_FORCEINLINE const uint8 *    VVM_Exec0null_exit_stat_scope              (VVM_NULL_FN_ARGS)                                                                                                                   { return InsPtr; }
 VM_FORCEINLINE const uint8 *    VVM_Exec0null_update_id                    (VVM_NULL_FN_ARGS)                                                                                                                   { return VVM_update_id(CT_MultipleLoops, InsPtr, BatchState, ExecCtx, NumLoops); }
diff --git a/Engine/Source/Runtime/VectorVM/Public/VectorVM.h b/Engine/Source/Runtime/VectorVM/Public/VectorVM.h
--- a/Engine/Source/Runtime/VectorVM/Public/VectorVM.h
+++ b/Engine/Source/Runtime/VectorVM/Public/VectorVM.h
@@ -84,7 +84,7 @@ enum class EVectorVMOperandLocation : uint8
 	VVM_OP_XM( sign                         , Op        , 1    , 1   , f        , 0             , 0               , VVM_INS_PARAM_FFFFFF) /* 33 */           \
 	VVM_OP_XM( step                         , Op        , 2    , 1   , f        , 0             , 0               , VVM_INS_PARAM_FFFFFF) /* 34 */           \
 	VVM_OP_XM( random                       , Op        , 1    , 1   , null     , 28            , 2               , VVM_INS_PARAM_FFFFFF) /* 35 */           \
-	VVM_OP_XM( noise                        , Op        , 0    , 0   , null     , 0             , 0               , VVM_INS_PARAM_FFFFFF) /* 36 */           \
+	VVM_OP_XM( noise                        , Op        , 1    , 1   , f        , 0             , 0               , VVM_INS_PARAM_FFFFFF) /* 36 */           \
 	VVM_OP_XM( cmplt                        , Op        , 2    , 1   , f        , 30            , 3               , VVM_INS_PARAM_FFFIFF) /* 37 */           \
 	VVM_OP_XM( cmple                        , Op        , 2    , 1   , f        , 33            , 3               , VVM_INS_PARAM_FFFIFF) /* 38 */           \
 	VVM_OP_XM( cmpgt                        , Op        , 2    , 1   , f        , 36            , 2               , VVM_INS_PARAM_FFFIFF) /* 39 */           \
@@ -137,8 +137,8 @@ enum class EVectorVMOperandLocation : uint8
 	VVM_OP_XM( acquireindex                 , IndexGen  , 1    , 0   , null     , 0             , 0               , VVM_INS_PARAM_FFFFII) /* 86 */           \
 	VVM_OP_XM( external_func_call           , ExtFnCall , 0    , 0   , null     , 0             , 0               , VVM_INS_PARAM_FFFFFF) /* 87 */           \
 	VVM_OP_XM( exec_index                   , Op        , 1    , 1   , null     , 78            , 2               , VVM_INS_PARAM_FFFFFI) /* 88 */           \
-	VVM_OP_XM( noise2D                      , Other     , 0    , 0   , null     , 0             , 0               , VVM_INS_PARAM_FFFFFF) /* 89 */           \
-	VVM_OP_XM( noise3D                      , Other     , 0    , 0   , null     , 0             , 0               , VVM_INS_PARAM_FFFFFF) /* 90 */           \
+	VVM_OP_XM( noise2D                      , Op		, 2    , 1   , f        , 0             , 0               , VVM_INS_PARAM_FFFFFF) /* 89 */           \
+	VVM_OP_XM( noise3D                      , Op		, 3    , 1   , f        , 0             , 0               , VVM_INS_PARAM_FFFFFF) /* 90 */           \
 	VVM_OP_XM( enter_stat_scope             , Stat      , 1    , 0   , null     , 0             , 0               , VVM_INS_PARAM_FFFFFF) /* 91 */           \
 	VVM_OP_XM( exit_stat_scope              , Stat      , 0    , 0   , null     , 0             , 0               , VVM_INS_PARAM_FFFFFF) /* 92 */           \
 	VVM_OP_XM( update_id                    , RWBuffer  , 0    , 0   , null     , 0             , 0               , VVM_INS_PARAM_FFFFFF) /* 93 */           \
